Fix UE5 Slate insights not displaying widget name due to bad merge.

diff --git a/Engine/Plugins/Slate/SlateInsights/Source/SlateInsights/Private/SlateProvider.cpp b/Engine/Plugins/Slate/SlateInsights/Source/SlateInsights/Private/SlateProvider.cpp
--- a/Engine/Plugins/Slate/SlateInsights/Source/SlateInsights/Private/SlateProvider.cpp
+++ b/Engine/Plugins/Slate/SlateInsights/Source/SlateInsights/Private/SlateProvider.cpp
@@ -35,8 +35,8 @@ namespace Message
 		, DebugInfo()
 		, EventIndex(0)
 	{
-		//EventData.GetString("Path", Path);
-		//EventData.GetString("DebugInfo", DebugInfo);
+		EventData.GetString("Path", Path);
+		EventData.GetString("DebugInfo", DebugInfo);
 	}
 	
 	FWidgetUpdatedMessage::FWidgetUpdatedMessage(const Trace::IAnalyzer::FEventData& EventData)
