Edigrate 14629275.
Capture callstack in widget invalidation for leaf widgets.

diff --git a/Engine/Plugins/Developer/VisualStudioSourceCodeAccess/VisualStudioSourceCodeAccess.uplugin b/Engine/Plugins/Developer/VisualStudioSourceCodeAccess/VisualStudioSourceCodeAccess.uplugin
--- a/Engine/Plugins/Developer/VisualStudioSourceCodeAccess/VisualStudioSourceCodeAccess.uplugin
+++ b/Engine/Plugins/Developer/VisualStudioSourceCodeAccess/VisualStudioSourceCodeAccess.uplugin
@@ -14,7 +14,7 @@
 	"CanContainContent" : false,
 	"IsBetaVersion" : false,
 	"Installed" : false,
-	"SupportedPrograms": [ "UnrealFrontend" ],
+	"SupportedPrograms": [ "UnrealFrontend", "UnrealInsights" ],
 	"Modules" :
 	[
 		{
@@ -26,7 +26,7 @@
 				"Win64",
 				"Win32"
 			],
-			"WhitelistPrograms" : [ "UnrealFrontend" ]
+			"WhitelistPrograms" : [ "UnrealFrontend", "UnrealInsights" ]
 		}
 	]
 }
\ No newline at end of file
diff --git a/Engine/Plugins/Developer/XCodeSourceCodeAccess/XCodeSourceCodeAccess.uplugin b/Engine/Plugins/Developer/XCodeSourceCodeAccess/XCodeSourceCodeAccess.uplugin
--- a/Engine/Plugins/Developer/XCodeSourceCodeAccess/XCodeSourceCodeAccess.uplugin
+++ b/Engine/Plugins/Developer/XCodeSourceCodeAccess/XCodeSourceCodeAccess.uplugin
@@ -14,7 +14,7 @@
 	"CanContainContent" : false,
 	"IsBetaVersion" : false,
 	"Installed" : false,
-	"SupportedPrograms": [ "UnrealFrontend" ],
+	"SupportedPrograms": [ "UnrealFrontend", "UnrealInsights" ],
 	"Modules" :
 	[
 		{
@@ -25,7 +25,7 @@
 			[
 				"Mac"
 			],
-			"WhitelistPrograms" : [ "UnrealFrontend" ]
+			"WhitelistPrograms" : [ "UnrealFrontend", "UnrealInsights" ]
 		}
 	]
 }
\ No newline at end of file
diff --git a/Engine/Plugins/Slate/SlateInsights/Source/SlateInsights/Private/SSlateFrameSchematicView.cpp b/Engine/Plugins/Slate/SlateInsights/Source/SlateInsights/Private/SSlateFrameSchematicView.cpp
--- a/Engine/Plugins/Slate/SlateInsights/Source/SlateInsights/Private/SSlateFrameSchematicView.cpp
+++ b/Engine/Plugins/Slate/SlateInsights/Source/SlateInsights/Private/SSlateFrameSchematicView.cpp
@@ -13,8 +13,18 @@
 #include "Widgets/Layout/SHeader.h"
 #include "Widgets/Text/STextBlock.h"
 #include "Widgets/Views/STableRow.h"
+#include "Widgets/Input/SMultiLineEditableTextBox.h"
+
 #include "Framework/Commands/UIAction.h"
 #include "Framework/MultiBox/MultiBoxBuilder.h"
+#include "Framework/Application/SlateApplication.h"
+
+#include "HAL/FileManager.h"
+#include "Modules/ModuleManager.h"
+#include "Internationalization/Regex.h"
+#include "SlateInsightsStyle.h"
+#include "ISourceCodeAccessModule.h"
+#include "ISourceCodeAccessor.h"
 
 #define LOCTEXT_NAMESPACE "SSlateFrameSchematicView"
 
@@ -40,6 +50,8 @@ namespace Private
 		uint32 Count;
 		TArray<TSharedPtr<FWidgetUniqueInvalidatedInfo>> Investigators;
 		bool bRoot;
+		FString ScriptTrace;
+		FString Callstack;
 	};
 
 	struct FWidgetUpdateInfo
@@ -81,6 +93,7 @@ namespace Private
 			if (ColumnName == ColumnWidgetId)
 			{
 				const EVisibility ExpanderArrowVisibility = Info->Investigators.Num() ? EVisibility::Visible : EVisibility::Hidden;
+				const FString TraceAndCallStackTipText = Info->ScriptTrace.IsEmpty() ? Info->Callstack : Info->ScriptTrace + "\n" + Info->Callstack;
 
 				return SNew(SHorizontalBox)
 
@@ -98,6 +111,7 @@ namespace Private
 					[
 						SNew(STextBlock)
 						.Text(WidgetName)
+						.ToolTipText(FText::AsCultureInvariant(TraceAndCallStackTipText))
 					];
 			}
 			if (ColumnName == ColumnNumber)
@@ -403,6 +417,15 @@ TSharedPtr<SWidget> SSlateFrameSchematicView::HandleWidgetInvalidateListContextM
 			FCanExecuteAction::CreateSP(this, &SSlateFrameSchematicView::CanWidgetInvalidateListGotoRootWidget)
 		));
 
+	MenuBuilder.AddMenuEntry(
+		LOCTEXT("ViewScriptAndCallStack", "View script and call stack"),
+		LOCTEXT("ViewScriptAndCallStackTooltip", "Open a window containing the script stack and the call stack. Script stack may be empty."),
+		FSlateIcon(),
+		FUIAction(
+			FExecuteAction::CreateSP(this, &SSlateFrameSchematicView::HandleWidgetInvalidateListViewScriptAndCallStack),
+			FCanExecuteAction::CreateSP(this, &SSlateFrameSchematicView::CanWidgetInvalidateListViewScriptAndCallStack)
+		));
+
 	return MenuBuilder.MakeWidget();
 }
 
@@ -438,6 +461,109 @@ void SSlateFrameSchematicView::HandleWidgetInvalidateListGotoRootWidget()
 	}
 }
 
+bool SSlateFrameSchematicView::CanWidgetInvalidateListViewScriptAndCallStack()
+{
+	if (WidgetInvalidateInfoListView->GetSelectedItems().Num() == 1)
+	{
+		TSharedPtr<Private::FWidgetUniqueInvalidatedInfo> SelectedItem = WidgetInvalidateInfoListView->GetSelectedItems().Last();
+		return !SelectedItem->Callstack.IsEmpty();
+	}
+
+	return false;
+}
+
+void SSlateFrameSchematicView::HandleWidgetInvalidateListViewScriptAndCallStack()
+{
+	if (WidgetInvalidateInfoListView->GetSelectedItems().Num() == 1)
+	{
+		TSharedPtr<Private::FWidgetUniqueInvalidatedInfo> Info = WidgetInvalidateInfoListView->GetSelectedItems().Last();
+		const FString TraceAndCallStackTipText = Info->ScriptTrace.IsEmpty() ? Info->Callstack : Info->ScriptTrace + "\n" + Info->Callstack;
+		const FText WidgetName = Private::GetWidgetName(AnalysisSession, Info->WidgetId);
+
+		/** Create the window to host our package dialog widget */
+		TSharedRef< SWindow > ViewScriptAndCallstackWindow = SNew(SWindow)
+		.Title(WidgetName)
+		.ClientSize(FVector2D(1200, 600))
+		.ActivationPolicy(EWindowActivationPolicy::Never)
+		.IsInitiallyMaximized(false)
+		[
+				
+			SNew(SBorder)
+			.BorderImage(FCoreStyle::Get().GetBrush("ToolPanel.GroupBorder"))
+			.OnMouseDoubleClick(this, &SSlateFrameSchematicView::OnMouseButtonDoubleClick)
+			[
+				SNew(SVerticalBox)
+				+ SVerticalBox::Slot()
+				.FillHeight(1.0f)
+				[
+					SAssignNew(ScriptAndCallStackTextBox, SMultiLineEditableTextBox)
+					.Style(FSlateInsightsStyle::Get(), "Callstack.TextBox")
+					.Text(FText::AsCultureInvariant(TraceAndCallStackTipText))
+					.SelectWordOnMouseDoubleClick(false)
+				]
+			]
+		];
+
+		/** Show the package dialog window as a modal window */
+		FSlateApplication::Get().AddModalWindow(ViewScriptAndCallstackWindow, FSlateApplication::Get().GetActiveModalWindow());
+	}
+}
+
+bool ExtractFilepathAndLineNumber(FString& PotentialFilePath, int32& LineNumber)
+{
+	// Extract filename and line number using regex	
+#if PLATFORM_WINDOWS
+	const FRegexPattern SourceCodeRegexPattern(TEXT("([a-zA-Z]:/[^:\\n\\r()]+(h|cpp)):([0-9]+)"));
+	const int32 LineNumberCaptureGroupID = 3;
+#else
+	const FRegexPattern SourceCodeRegexPattern(TEXT("((//([^:/\\n]+[/])*)([^/]+)(h|cpp)):([0-9]+))"));
+	const int32 LineNumberCaptureGroupID = 6;
+#endif
+
+	FRegexMatcher SourceCodeRegexMatcher(SourceCodeRegexPattern, PotentialFilePath);
+	if (SourceCodeRegexMatcher.FindNext())
+	{
+		PotentialFilePath = SourceCodeRegexMatcher.GetCaptureGroup(1);
+		LineNumber = FCString::Strtoi(*SourceCodeRegexMatcher.GetCaptureGroup(LineNumberCaptureGroupID), nullptr, 10);
+		return true;
+	}
+
+	return false;
+}
+
+FReply SSlateFrameSchematicView::OnMouseButtonDoubleClick(const FGeometry& InMyGeometry, const FPointerEvent& InMouseEvent)
+{
+	if (InMouseEvent.GetEffectingButton() == EKeys::LeftMouseButton)
+	{
+		// grab cursor location's line of text
+		FString PotentialCodeFilePath;
+		if (ScriptAndCallStackTextBox)
+		{
+			ScriptAndCallStackTextBox->GetCurrentTextLine(PotentialCodeFilePath);
+
+			int32 LeftBrace = INDEX_NONE;
+			int32 RightBrace = INDEX_NONE;
+			if (PotentialCodeFilePath.FindChar('[', LeftBrace) && PotentialCodeFilePath.FindChar(']', RightBrace))
+			{
+				PotentialCodeFilePath.MidInline(LeftBrace + 1, RightBrace - LeftBrace - 1);
+
+				// Extract potential .cpp./h files file path & line number
+				int32 LineNumber = 0;
+				PotentialCodeFilePath = IFileManager::Get().ConvertToAbsolutePathForExternalAppForRead(*PotentialCodeFilePath);
+				if (ExtractFilepathAndLineNumber(PotentialCodeFilePath, LineNumber) && PotentialCodeFilePath.Len() && IFileManager::Get().FileSize(*PotentialCodeFilePath) != INDEX_NONE)
+				{
+					ISourceCodeAccessModule& SourceCodeAccessModule = FModuleManager::LoadModuleChecked<ISourceCodeAccessModule>("SourceCodeAccess");
+					SourceCodeAccessModule.GetAccessor().OpenFileAtLine(PotentialCodeFilePath, LineNumber);
+				}
+
+				return FReply::Handled();
+			}
+		}
+	}
+
+	return FReply::Unhandled();
+}
+
 void SSlateFrameSchematicView::RefreshNodes()
 {
 	WidgetInvalidationInfos.Reset();
@@ -534,6 +660,15 @@ void SSlateFrameSchematicView::RefreshNodes_Invalidation(const FSlateProvider* S
 				WidgetInfo->Investigators.Add(InvestigatorInfo);
 			}
 
+			if (!Message.ScriptTrace.IsEmpty())
+			{
+				WidgetInfo->ScriptTrace = Message.ScriptTrace;
+			}
+
+			if (!Message.Callstack.IsEmpty())
+			{
+				WidgetInfo->Callstack = Message.Callstack;
+			}
 
 			return Trace::EEventEnumerate::Continue;
 		});
diff --git a/Engine/Plugins/Slate/SlateInsights/Source/SlateInsights/Private/SSlateFrameSchematicView.h b/Engine/Plugins/Slate/SlateInsights/Source/SlateInsights/Private/SSlateFrameSchematicView.h
--- a/Engine/Plugins/Slate/SlateInsights/Source/SlateInsights/Private/SSlateFrameSchematicView.h
+++ b/Engine/Plugins/Slate/SlateInsights/Source/SlateInsights/Private/SSlateFrameSchematicView.h
@@ -17,6 +17,7 @@ class ITableRow;
 class ITimingEvent;
 class STableViewBase;
 class STextBlock;
+class SMultiLineEditableTextBox;
 
 namespace UE
 {
@@ -45,9 +46,14 @@ private:
 	void HandleTimeMarkerChanged(Insights::ETimeChangedFlags InFlags, double InTimeMarker);
 	void HandleSelectionChanged(Insights::ETimeChangedFlags InFlags, double StartTime, double EndTime);
 	void HandleSelectionEventChanged(const TSharedPtr<const ITimingEvent> InEvent);
+
 	TSharedPtr<SWidget> HandleWidgetInvalidateListContextMenu();
 	bool CanWidgetInvalidateListGotoRootWidget();
 	void HandleWidgetInvalidateListGotoRootWidget();
+	bool CanWidgetInvalidateListViewScriptAndCallStack();
+	void HandleWidgetInvalidateListViewScriptAndCallStack();
+
+	virtual FReply OnMouseButtonDoubleClick(const FGeometry& InMyGeometry, const FPointerEvent& InMouseEvent) override;
 
 	void RefreshNodes();
 	void RefreshNodes_Invalidation(const FSlateProvider* SlateProvider);
@@ -65,6 +71,7 @@ private:
 
 	TSharedPtr<STextBlock> InvalidationSummary;
 	TSharedPtr<STextBlock> UpdateSummary;
+	TSharedPtr<SMultiLineEditableTextBox> ScriptAndCallStackTextBox;
 
 	double StartTime;
 	double EndTime;
diff --git a/Engine/Plugins/Slate/SlateInsights/Source/SlateInsights/Private/SlateInsightsStyle.h b/Engine/Plugins/Slate/SlateInsights/Source/SlateInsights/Private/SlateInsightsStyle.h
--- a/Engine/Plugins/Slate/SlateInsights/Source/SlateInsights/Private/SlateInsightsStyle.h
+++ b/Engine/Plugins/Slate/SlateInsights/Source/SlateInsights/Private/SlateInsightsStyle.h
@@ -5,7 +5,6 @@
 #include "Styling/CoreStyle.h"
 #include "Styling/SlateStyle.h"
 #include "Styling/SlateStyleRegistry.h"
-#include "Styling/CoreStyle.h"
 
 namespace UE
 {
@@ -41,6 +40,11 @@ public:
 		Set("Flag.Color.Background", FCoreStyle::Get().GetSlateColor("InvertedForeground"));
 		Set("Flag.Color.Selected", FCoreStyle::Get().GetSlateColor("SelectionColor"));
 
+		Set("Callstack.TextBox", FEditableTextBoxStyle(FCoreStyle::Get().GetWidgetStyle<FEditableTextBoxStyle>("DarkEditableTextBox"))
+			.SetBackgroundColor(FLinearColor(0.015f, 0.015f, 0.015f))
+			.SetForegroundColor(FLinearColor::White * 0.8f)
+		);
+
 		FSlateStyleRegistry::RegisterSlateStyle(*this);
 	}
 
diff --git a/Engine/Plugins/Slate/SlateInsights/Source/SlateInsights/Private/SlateProvider.cpp b/Engine/Plugins/Slate/SlateInsights/Source/SlateInsights/Private/SlateProvider.cpp
--- a/Engine/Plugins/Slate/SlateInsights/Source/SlateInsights/Private/SlateProvider.cpp
+++ b/Engine/Plugins/Slate/SlateInsights/Source/SlateInsights/Private/SlateProvider.cpp
@@ -1,6 +1,8 @@
 // Copyright Epic Games, Inc. All Rights Reserved.
 
 #include "SlateProvider.h"
+#include "HAL/PlatformStackWalk.h"
+#include "HAL/PlatformProcess.h"
 #include "Insights/ViewModels/TimingEventsTrack.h"
 
 #define LOCTEXT_NAMESPACE "SlateProvider"
@@ -54,6 +56,8 @@ namespace Message
 		Message.WidgetId = EventData.GetValue<uint64>("WidgetId");
 		Message.InvestigatorId = EventData.GetValue<uint64>("InvestigatorId");
 		Message.InvalidationReason = static_cast<EInvalidateWidgetReason>(EventData.GetValue<uint8>("InvalidateWidgetReason"));
+		EventData.GetString("ScriptTrace", Message.ScriptTrace);
+		Message.Callstack = GetCallstack(EventData);
 		return Message;
 	}
 
@@ -63,6 +67,8 @@ namespace Message
 		Message.WidgetId = EventData.GetValue<uint64>("WidgetId");
 		Message.InvestigatorId = EventData.GetValue<uint64>("InvestigatorId");
 		Message.InvalidationReason = EInvalidateWidgetReason::Layout;
+		EventData.GetString("ScriptTrace", Message.ScriptTrace);
+		Message.Callstack = GetCallstack(EventData);
 		return Message;
 	}
 
@@ -72,8 +78,37 @@ namespace Message
 		Message.WidgetId = EventData.GetValue<uint64>("WidgetId");
 		Message.InvestigatorId = EventData.GetValue<uint64>("InvestigatorId");
 		Message.InvalidationReason = EInvalidateWidgetReason::ChildOrder;
+		EventData.GetString("ScriptTrace", Message.ScriptTrace);
+		Message.Callstack = GetCallstack(EventData);
 		return Message;
 	}
+
+	FString FWidgetInvalidatedMessage::GetCallstack(const Trace::IAnalyzer::FEventData& EventData)
+	{
+		FString Callstack = "";
+		const auto& CallstackReader = EventData.GetArrayView<uint64>("Callstack");
+		if (CallstackReader.Num() != 0)
+		{
+			FPlatformStackWalk::InitStackWalkingForProcess(FPlatformProcess::OpenProcess(EventData.GetValue<uint32>("ProcessId")));
+			uint8 StackDepth = 0;
+			for (uint64 ProgramCounter : CallstackReader)
+			{
+				// Skip the first two backraces, that's from us.
+				if (StackDepth++ >= 2)
+				{
+					FString SymbolAsText;
+
+					// Note: Done insight thread as this process is very slow, possibly 1~80ms based on widget complexity.
+					FProgramCounterSymbolInfoEx SymbolInfo;
+					FPlatformStackWalk::ProgramCounterToSymbolInfoEx(ProgramCounter, SymbolInfo);
+					FPlatformStackWalk::SymbolInfoToHumanReadableStringEx(SymbolInfo, SymbolAsText);
+
+					Callstack += SymbolAsText + "\r\n";
+				}
+			}
+		}
+		return Callstack;
+	}
 } //namespace Message
 
 FSlateProvider::FSlateProvider(Trace::IAnalysisSession& InSession)
diff --git a/Engine/Plugins/Slate/SlateInsights/Source/SlateInsights/Private/SlateProvider.h b/Engine/Plugins/Slate/SlateInsights/Source/SlateInsights/Private/SlateProvider.h
--- a/Engine/Plugins/Slate/SlateInsights/Source/SlateInsights/Private/SlateProvider.h
+++ b/Engine/Plugins/Slate/SlateInsights/Source/SlateInsights/Private/SlateProvider.h
@@ -69,10 +69,13 @@ struct FWidgetInvalidatedMessage
 	EInvalidateWidgetReason InvalidationReason = EInvalidateWidgetReason::None;
 	bool bRootInvalidated = false;
 	bool bRootChildOrderInvalidated = false;
+	FString ScriptTrace;
+	FString Callstack;
 
 	static FWidgetInvalidatedMessage FromWidget(const Trace::IAnalyzer::FEventData& EventData);
 	static FWidgetInvalidatedMessage FromRoot(const Trace::IAnalyzer::FEventData& EventData);
 	static FWidgetInvalidatedMessage FromChildOrder(const Trace::IAnalyzer::FEventData& EventData);
+	static FString GetCallstack(const Trace::IAnalyzer::FEventData& EventData);
 };
 
 struct FApplicationTickedMessage
diff --git a/Engine/Plugins/Slate/SlateInsights/Source/SlateInsights/SlateInsights.Build.cs b/Engine/Plugins/Slate/SlateInsights/Source/SlateInsights/SlateInsights.Build.cs
--- a/Engine/Plugins/Slate/SlateInsights/Source/SlateInsights/SlateInsights.Build.cs
+++ b/Engine/Plugins/Slate/SlateInsights/Source/SlateInsights/SlateInsights.Build.cs
@@ -10,6 +10,7 @@ namespace UnrealBuildTool.Rules
 			new string[]
 			{
 				"Core",
+				"CoreUObject",
 				"InputCore",
 				"SlateCore",
 				"Slate",
@@ -19,6 +20,7 @@ namespace UnrealBuildTool.Rules
 				"TraceInsights",
 				"AssetRegistry",
 				"ApplicationCore",
+				"SourceCodeAccess",
 			});
 
 			if (Target.bCompileAgainstEngine)
@@ -27,7 +29,6 @@ namespace UnrealBuildTool.Rules
 				new string[]
 				{
 					"Engine",
-					"CoreUObject",
 				});
 			}
 		}
diff --git a/Engine/Programs/UnrealInsights/Config/DefaultEngine.ini b/Engine/Programs/UnrealInsights/Config/DefaultEngine.ini
--- a/Engine/Programs/UnrealInsights/Config/DefaultEngine.ini
+++ b/Engine/Programs/UnrealInsights/Config/DefaultEngine.ini
@@ -2,4 +2,6 @@
 +ProgramEnabledPlugins=GameplayInsights
 +ProgramEnabledPlugins=TraceDataFilters
 +ProgramEnabledPlugins=SlateInsights
-+ProgramEnabledPlugins=NetworkPredictionInsights
\ No newline at end of file
++ProgramEnabledPlugins=NetworkPredictionInsights
++ProgramEnabledPlugins=VisualStudioSourceCodeAccess
++ProgramEnabledPlugins=XcodeSourceCodeAccess
diff --git a/Engine/Source/Programs/UnrealInsights/Private/UserInterfaceCommand.cpp b/Engine/Source/Programs/UnrealInsights/Private/UserInterfaceCommand.cpp
--- a/Engine/Source/Programs/UnrealInsights/Private/UserInterfaceCommand.cpp
+++ b/Engine/Source/Programs/UnrealInsights/Private/UserInterfaceCommand.cpp
@@ -161,11 +161,11 @@ void FUserInterfaceCommand::Run()
 
 	// Manually load in the source code access plugins, as standalone programs don't currently support plugins.
 #if PLATFORM_MAC
-	//im:TODO: IModuleInterface& XCodeSourceCodeAccessModule = FModuleManager::LoadModuleChecked<IModuleInterface>(FName("XCodeSourceCodeAccess"));
-	//im:TODO: SourceCodeAccessModule.SetAccessor(FName("XCodeSourceCodeAccess"));
+	IModuleInterface& XCodeSourceCodeAccessModule = FModuleManager::LoadModuleChecked<IModuleInterface>(FName("XCodeSourceCodeAccess"));
+	SourceCodeAccessModule.SetAccessor(FName("XCodeSourceCodeAccess"));
 #elif PLATFORM_WINDOWS
-	//im:TODO: IModuleInterface& VisualStudioSourceCodeAccessModule = FModuleManager::LoadModuleChecked<IModuleInterface>(FName("VisualStudioSourceCodeAccess"));
-	//im:TODO: SourceCodeAccessModule.SetAccessor(FName("VisualStudioSourceCodeAccess"));
+	IModuleInterface& VisualStudioSourceCodeAccessModule = FModuleManager::LoadModuleChecked<IModuleInterface>(FName("VisualStudioSourceCodeAccess"));
+	SourceCodeAccessModule.SetAccessor(FName("VisualStudioSourceCodeAccess"));
 #endif
 
 #if WITH_SHARED_POINTER_TESTS
diff --git a/Engine/Source/Programs/UnrealInsights/UnrealInsights.Target.cs b/Engine/Source/Programs/UnrealInsights/UnrealInsights.Target.cs
--- a/Engine/Source/Programs/UnrealInsights/UnrealInsights.Target.cs
+++ b/Engine/Source/Programs/UnrealInsights/UnrealInsights.Target.cs
@@ -20,6 +20,10 @@ public class UnrealInsightsTarget : TargetRules
 		bCompileWithStatsWithoutEngine = true;
 		bCompileWithPluginSupport = true;
 
+		// For source code editor access & regex (crossplatform)
+		bIncludePluginsForTargetPlatforms = true;
+		bCompileICU = true;
+
 		// For UI functionality
 		bBuildDeveloperTools = true;
 
diff --git a/Engine/Source/Runtime/Core/Private/GenericPlatform/GenericPlatformCrashContext.cpp b/Engine/Source/Runtime/Core/Private/GenericPlatform/GenericPlatformCrashContext.cpp
--- a/Engine/Source/Runtime/Core/Private/GenericPlatform/GenericPlatformCrashContext.cpp
+++ b/Engine/Source/Runtime/Core/Private/GenericPlatform/GenericPlatformCrashContext.cpp
@@ -1187,17 +1187,6 @@ bool FGenericCrashContext::CreateCrashReportDirectory(const TCHAR* CrashGUIDRoot
 	return IFileManager::Get().MakeDirectory(*OutCrashDirectoryAbsolute, true);
 }
 
-FProgramCounterSymbolInfoEx::FProgramCounterSymbolInfoEx( FString InModuleName, FString InFunctionName, FString InFilename, uint32 InLineNumber, uint64 InSymbolDisplacement, uint64 InOffsetInModule, uint64 InProgramCounter ) :
-	ModuleName( InModuleName ),
-	FunctionName( InFunctionName ),
-	Filename( InFilename ),
-	LineNumber( InLineNumber ),
-	SymbolDisplacement( InSymbolDisplacement ),
-	OffsetInModule( InOffsetInModule ),
-	ProgramCounter( InProgramCounter )
-{
-}
-
 FString RecoveryService::GetRecoveryServerName()
 {
 	// Requirements: To avoid collision, the name must be unique on the local machine (multiple instances) and across the local network (multiple users).
diff --git a/Engine/Source/Runtime/Core/Private/GenericPlatform/GenericPlatformStackWalk.cpp b/Engine/Source/Runtime/Core/Private/GenericPlatform/GenericPlatformStackWalk.cpp
--- a/Engine/Source/Runtime/Core/Private/GenericPlatform/GenericPlatformStackWalk.cpp
+++ b/Engine/Source/Runtime/Core/Private/GenericPlatform/GenericPlatformStackWalk.cpp
@@ -16,6 +16,17 @@ FProgramCounterSymbolInfo::FProgramCounterSymbolInfo() :
 	FCStringAnsi::Strncpy( Filename, "", MAX_NAME_LENGTH );
 }
 
+FProgramCounterSymbolInfoEx::FProgramCounterSymbolInfoEx(FString InModuleName, FString InFunctionName, FString InFilename, uint32 InLineNumber, uint64 InSymbolDisplacement, uint64 InOffsetInModule, uint64 InProgramCounter) :
+	ModuleName(InModuleName),
+	FunctionName(InFunctionName),
+	Filename(InFilename),
+	LineNumber(InLineNumber),
+	SymbolDisplacement(InSymbolDisplacement),
+	OffsetInModule(InOffsetInModule),
+	ProgramCounter(InProgramCounter)
+{
+}
+
 /** Settings for stack walking */
 static bool GWantsDetailedCallstacksInNonMonolithicBuilds = true;
 
@@ -126,10 +137,12 @@ bool FGenericPlatformStackWalk::SymbolInfoToHumanReadableStringEx( const FProgra
 	// ModuleName! {ProgramCounter}
 	
 	// Strip module path.
-	const TCHAR* Pos0 = FCString::Strrchr( *SymbolInfo.ModuleName, '\\' );
-	const TCHAR* Pos1 = FCString::Strrchr( *SymbolInfo.ModuleName, '/' );
-	const UPTRINT RealPos = FMath::Max( (UPTRINT)Pos0, (UPTRINT)Pos1 );
-	const FString StrippedModuleName = RealPos > 0 ? (const ANSICHAR*)(RealPos + 1) : SymbolInfo.ModuleName;
+	int32 Pos0;
+	int32 Pos1;
+	SymbolInfo.ModuleName.FindLastChar('\\', Pos0);
+	SymbolInfo.ModuleName.FindLastChar('/', Pos1);
+	const int32 RealPos = FMath::Max( Pos0, Pos1 );
+	const FString StrippedModuleName = RealPos > 0 ? SymbolInfo.ModuleName.RightChop(RealPos + 1) : SymbolInfo.ModuleName;
 
 	out_HumanReadableString = StrippedModuleName;
 	
diff --git a/Engine/Source/Runtime/Core/Private/Windows/WindowsPlatformStackWalk.cpp b/Engine/Source/Runtime/Core/Private/Windows/WindowsPlatformStackWalk.cpp
--- a/Engine/Source/Runtime/Core/Private/Windows/WindowsPlatformStackWalk.cpp
+++ b/Engine/Source/Runtime/Core/Private/Windows/WindowsPlatformStackWalk.cpp
@@ -493,6 +493,78 @@ void FWindowsPlatformStackWalk::ProgramCounterToSymbolInfo( uint64 ProgramCounte
 	}
 }
 
+void FWindowsPlatformStackWalk::ProgramCounterToSymbolInfoEx(uint64 ProgramCounter, FProgramCounterSymbolInfoEx& out_SymbolInfo)
+{
+#if ON_DEMAND_SYMBOL_LOADING
+	// Load symbols for the module
+	LoadSymbolsForModuleByAddress(ProgramCounter, GetSymbolSearchPath());
+#endif
+
+	// Set the program counter.
+	out_SymbolInfo.ProgramCounter = ProgramCounter;
+
+	uint32 LastError = 0;
+	HANDLE ProcessHandle = GProcessHandle;
+
+	// Initialize symbol.
+	ANSICHAR SymbolBuffer[sizeof(IMAGEHLP_SYMBOL64) + FProgramCounterSymbolInfo::MAX_NAME_LENGTH] = { 0 };
+	IMAGEHLP_SYMBOL64* Symbol = (IMAGEHLP_SYMBOL64*)SymbolBuffer;
+	Symbol->SizeOfStruct = sizeof(SymbolBuffer);
+	Symbol->MaxNameLength = FProgramCounterSymbolInfo::MAX_NAME_LENGTH;
+
+	// Get function name.
+	ANSICHAR FunctionName[FProgramCounterSymbolInfo::MAX_NAME_LENGTH] = { 0 };
+	if (SymGetSymFromAddr64(ProcessHandle, ProgramCounter, nullptr, Symbol))
+	{
+		// Skip any funky chars in the beginning of a function name.
+		int32 Offset = 0;
+		while (Symbol->Name[Offset] < 32 || Symbol->Name[Offset] > 127)
+		{
+			Offset++;
+		}
+
+		// Write out function name.
+		FCStringAnsi::Strncpy(FunctionName, Symbol->Name + Offset, FProgramCounterSymbolInfo::MAX_NAME_LENGTH);
+		FCStringAnsi::Strncat(FunctionName, "()", FProgramCounterSymbolInfo::MAX_NAME_LENGTH);
+		out_SymbolInfo.FunctionName = ANSI_TO_TCHAR(FunctionName);
+	}
+	else
+	{
+		// No symbol found for this address.
+		LastError = GetLastError();
+	}
+
+	// Get filename and line number.
+	ANSICHAR Filename[FProgramCounterSymbolInfo::MAX_NAME_LENGTH] = { 0 };
+	IMAGEHLP_LINE64	ImageHelpLine = { 0 };
+	ImageHelpLine.SizeOfStruct = sizeof(ImageHelpLine);
+	if (SymGetLineFromAddr64(ProcessHandle, ProgramCounter, (::DWORD*) & out_SymbolInfo.SymbolDisplacement, &ImageHelpLine))
+	{
+		FCStringAnsi::Strncpy(Filename, ImageHelpLine.FileName, FProgramCounterSymbolInfo::MAX_NAME_LENGTH);
+		out_SymbolInfo.Filename = ANSI_TO_TCHAR(Filename);
+		out_SymbolInfo.LineNumber = ImageHelpLine.LineNumber;
+	}
+	else
+	{
+		LastError = GetLastError();
+	}
+
+	// Get module name.
+	ANSICHAR ModuleName[FProgramCounterSymbolInfo::MAX_NAME_LENGTH] = { 0 };
+	IMAGEHLP_MODULE64 ImageHelpModule = { 0 };
+	ImageHelpModule.SizeOfStruct = sizeof(ImageHelpModule);
+	if (SymGetModuleInfo64(ProcessHandle, ProgramCounter, &ImageHelpModule))
+	{
+		// Write out module information.
+		FCStringAnsi::Strncpy(ModuleName, ImageHelpModule.ImageName, FProgramCounterSymbolInfo::MAX_NAME_LENGTH);
+		out_SymbolInfo.ModuleName = ANSI_TO_TCHAR(ModuleName);
+	}
+	else
+	{
+		LastError = GetLastError();
+	}
+}
+
 /**
  * Get process module handle NULL-terminated list.
  * On error this method returns NULL.
diff --git a/Engine/Source/Runtime/Core/Public/GenericPlatform/GenericPlatformCrashContext.h b/Engine/Source/Runtime/Core/Public/GenericPlatform/GenericPlatformCrashContext.h
--- a/Engine/Source/Runtime/Core/Public/GenericPlatform/GenericPlatformCrashContext.h
+++ b/Engine/Source/Runtime/Core/Public/GenericPlatform/GenericPlatformCrashContext.h
@@ -5,11 +5,10 @@
 #include "CoreTypes.h"
 #include "HAL/PlatformMemory.h"
 #include "HAL/PlatformProcess.h"
+#include "HAL/PlatformStackwalk.h"
 #include "Misc/Optional.h"
 #include "Containers/UnrealString.h"
 
-struct FProgramCounterSymbolInfo;
-
 /** Defines special exit codes used to diagnose abnormal terminations. The code values are arbitrary, but easily recongnizable in decimal. They are meant to be
     used with the out-of-process monitoring/analytics in order to figure out unexpected cases. */
 enum ECrashExitCodes : int32
@@ -39,39 +38,6 @@ enum ECrashExitCodes : int32
 	UnhandledEnsure = 777008,
 };
 
-/**
- * Symbol information associated with a program counter. 
- * FString version.
- * To be used by external tools.
- */
-struct CORE_API FProgramCounterSymbolInfoEx
-{
-	/** Module name. */
-	FString	ModuleName;
-
-	/** Function name. */
-	FString	FunctionName;
-
-	/** Filename. */
-	FString	Filename;
-
-	/** Line number in file. */
-	uint32	LineNumber;
-
-	/** Symbol displacement of address.	*/
-	uint64	SymbolDisplacement;
-
-	/** Program counter offset into module. */
-	uint64	OffsetInModule;
-
-	/** Program counter. */
-	uint64	ProgramCounter;
-
-	/** Default constructor. */
-	FProgramCounterSymbolInfoEx( FString InModuleName, FString InFunctionName, FString InFilename, uint32 InLineNumber, uint64 InSymbolDisplacement, uint64 InOffsetInModule, uint64 InProgramCounter );
-};
-
-
 /** Enumerates crash description versions. */
 enum class ECrashDescVersions : int32
 {
diff --git a/Engine/Source/Runtime/Core/Public/GenericPlatform/GenericPlatformStackWalk.h b/Engine/Source/Runtime/Core/Public/GenericPlatform/GenericPlatformStackWalk.h
--- a/Engine/Source/Runtime/Core/Public/GenericPlatform/GenericPlatformStackWalk.h
+++ b/Engine/Source/Runtime/Core/Public/GenericPlatform/GenericPlatformStackWalk.h
@@ -71,7 +71,32 @@ struct FProgramCounterSymbolInfo final
 	CORE_API FProgramCounterSymbolInfo();
 };
 
-struct FProgramCounterSymbolInfoEx;
+struct CORE_API FProgramCounterSymbolInfoEx
+{
+	/** Module name. */
+	FString	ModuleName;
+
+	/** Function name. */
+	FString	FunctionName;
+
+	/** Filename. */
+	FString	Filename;
+
+	/** Line number in file. */
+	uint32	LineNumber;
+
+	/** Symbol displacement of address.	*/
+	uint64	SymbolDisplacement;
+
+	/** Program counter offset into module. */
+	uint64	OffsetInModule;
+
+	/** Program counter. */
+	uint64	ProgramCounter;
+
+	/** Default constructor. */
+	FProgramCounterSymbolInfoEx(FString InModuleName = "", FString InFunctionName = "", FString InFilename = "", uint32 InLineNumber = 0, uint64 InSymbolDisplacement = 0, uint64 InOffsetInModule = 0, uint64 InProgramCounter = 0);
+};
 
 /**
  * Generic implementation for most platforms
@@ -174,6 +199,12 @@ struct CORE_API FGenericPlatformStackWalk
 		out_SymbolInfo.ProgramCounter = ProgramCounter;
 	}
 
+	/** Same as above, but can be used with external applications since it doesn't re-initialize the active process. */
+	static void ProgramCounterToSymbolInfoEx( uint64 ProgramCounter, FProgramCounterSymbolInfoEx& out_SymbolInfo)
+	{
+		out_SymbolInfo.ProgramCounter = ProgramCounter;
+	}
+
 	/**
 	 * Capture a stack backtrace and optionally use the passed in exception pointers.
 	 *
diff --git a/Engine/Source/Runtime/Core/Public/Windows/WindowsPlatformStackWalk.h b/Engine/Source/Runtime/Core/Public/Windows/WindowsPlatformStackWalk.h
--- a/Engine/Source/Runtime/Core/Public/Windows/WindowsPlatformStackWalk.h
+++ b/Engine/Source/Runtime/Core/Public/Windows/WindowsPlatformStackWalk.h
@@ -17,6 +17,7 @@ struct CORE_API FWindowsPlatformStackWalk
 	static TArray<FProgramCounterSymbolInfo> GetStack(int32 IgnoreCount, int32 MaxDepth = 100, void* Context = nullptr);
 
 	static void ProgramCounterToSymbolInfo( uint64 ProgramCounter, FProgramCounterSymbolInfo& out_SymbolInfo );
+	static void ProgramCounterToSymbolInfoEx( uint64 ProgramCounter, FProgramCounterSymbolInfoEx& out_SymbolInfo );
 	static uint32 CaptureStackBackTrace( uint64* BackTrace, uint32 MaxDepth, void* Context = nullptr );
 	static uint32 CaptureThreadStackBackTrace( uint64 ThreadId, uint64* BackTrace, uint32 MaxDepth );
 	static void StackWalkAndDump( ANSICHAR* HumanReadableString, SIZE_T HumanReadableStringSize, int32 IgnoreCount, void* Context = nullptr );
diff --git a/Engine/Source/Runtime/CoreUObject/Private/UObject/ScriptCore.cpp b/Engine/Source/Runtime/CoreUObject/Private/UObject/ScriptCore.cpp
--- a/Engine/Source/Runtime/CoreUObject/Private/UObject/ScriptCore.cpp
+++ b/Engine/Source/Runtime/CoreUObject/Private/UObject/ScriptCore.cpp
@@ -456,7 +456,7 @@ static bool ShowKismetScriptStackOnWarnings()
 	return ShowScriptStackForScriptWarning;
 }
 
-FString FFrame::GetScriptCallstack()
+FString FFrame::GetScriptCallstack(bool bReturnEmpty)
 {
 	FString ScriptStack;
 
@@ -469,12 +469,15 @@ FString FFrame::GetScriptCallstack()
 			ScriptStack += TEXT("\t") + BlueprintExceptionTracker.ScriptStack[i]->GetStackDescription() + TEXT("\n");
 		}
 	}
-	else
+	else if (!bReturnEmpty)
 	{
 		ScriptStack += TEXT("\t[Empty] (FFrame::GetScriptCallstack() called from native code)");
 	}
 #else
-	ScriptStack = TEXT("Unable to display Script Callstack. Compile with DO_BLUEPRINT_GUARD=1");
+	if (!bReturnEmpty)
+	{
+		ScriptStack = TEXT("Unable to display Script Callstack. Compile with DO_BLUEPRINT_GUARD=1");
+	}
 #endif
 
 	return ScriptStack;
diff --git a/Engine/Source/Runtime/CoreUObject/Public/UObject/Stack.h b/Engine/Source/Runtime/CoreUObject/Public/UObject/Stack.h
--- a/Engine/Source/Runtime/CoreUObject/Public/UObject/Stack.h
+++ b/Engine/Source/Runtime/CoreUObject/Public/UObject/Stack.h
@@ -187,8 +187,10 @@ public:
 
 	/**
 	* This will return the StackTrace of the all script frames currently active
+	* 
+	* @param	bReturnEmpty if true, returns empty string when no script callstack found
 	**/
-	COREUOBJECT_API static FString GetScriptCallstack();
+	COREUOBJECT_API static FString GetScriptCallstack(bool bReturnEmpty = false);
 
 	/** 
 	 * This will return a string of the form "ScopeName.FunctionName" associated with this stack frame:
diff --git a/Engine/Source/Runtime/SlateCore/Private/Trace/SlateTrace.cpp b/Engine/Source/Runtime/SlateCore/Private/Trace/SlateTrace.cpp
--- a/Engine/Source/Runtime/SlateCore/Private/Trace/SlateTrace.cpp
+++ b/Engine/Source/Runtime/SlateCore/Private/Trace/SlateTrace.cpp
@@ -6,11 +6,24 @@
 
 #include "Application/SlateApplicationBase.h"
 #include "HAL/PlatformTime.h"
+#include "HAL/ThreadHeartBeat.h"
+#include "HAL/PlatformStackWalk.h"
+#include "HAL/PlatformProcess.h"
 #include "Trace/Trace.inl"
 #include "Types/ReflectionMetadata.h"
 #include "Widgets/SWidget.h"
 #include "FastUpdate/WidgetProxy.h"
 
+#if !(UE_BUILD_SHIPPING)
+
+static int32 bCaptureRootInvalidationCallstacks = 0;
+static FAutoConsoleVariableRef CVarCaptureRootInvalidationCallstacks(
+	TEXT("SlateDebugger.bCaptureRootInvalidationCallstacks"),
+	bCaptureRootInvalidationCallstacks,
+	TEXT("Whenever a widget is the root cause of an invalidation, capture the callstack for slate insights."));
+
+#endif // !(UE_BUILD_SHIPPING)
+
 //-----------------------------------------------------------------------------------//
 
 UE_TRACE_CHANNEL_DEFINE(SlateChannel)
@@ -55,9 +68,12 @@ UE_TRACE_EVENT_END()
 
 UE_TRACE_EVENT_BEGIN(SlateTrace, WidgetInvalidated)
 	UE_TRACE_EVENT_FIELD(uint64, Cycle)
-	UE_TRACE_EVENT_FIELD(uint64, WidgetId)				// Invalidated widget unique ID.
-	UE_TRACE_EVENT_FIELD(uint64, InvestigatorId)		// Widget unique ID that investigated the invalidation.
-	UE_TRACE_EVENT_FIELD(uint8, InvalidateWidgetReason)	// The reason of the invalidation. (EInvalidateWidgetReason)
+	UE_TRACE_EVENT_FIELD(uint64, WidgetId)					// Invalidated widget unique ID.
+	UE_TRACE_EVENT_FIELD(uint64, InvestigatorId)			// Widget unique ID that investigated the invalidation.
+	UE_TRACE_EVENT_FIELD(uint8, InvalidateWidgetReason)		// The reason of the invalidation. (EInvalidateWidgetReason)
+	UE_TRACE_EVENT_FIELD(Trace::WideString, ScriptTrace)	// Optional script trace for root widget invalidations
+	UE_TRACE_EVENT_FIELD(uint64[], Callstack)				// Optional callstack for root widget invalidations
+	UE_TRACE_EVENT_FIELD(uint32, ProcessId)					// Optional proccess ID where the invalidation occureed.
 UE_TRACE_EVENT_END()
 
 UE_TRACE_EVENT_BEGIN(SlateTrace, RootInvalidated)
@@ -186,6 +202,31 @@ void FSlateTrace::WidgetInvalidated(const SWidget* Widget, const SWidget* Invest
 
 		static_assert(sizeof(EInvalidateWidgetReason) == sizeof(uint8), "EInvalidateWidgetReason is not a uint8");
 
+		FString ScriptTrace;
+		constexpr int MAX_DEPTH = 64;
+		uint64 StackTrace[MAX_DEPTH] = { 0 };
+		uint32 StackTraceDepth = 0;
+		uint32 ProcessId = 0;
+
+#if !(UE_BUILD_SHIPPING)
+		//@TODO: Could add a CVar to only capture certain callstacks for performance (Widget name, type, etc).
+		if (!Investigator && bCaptureRootInvalidationCallstacks)
+		{
+			FSlowHeartBeatScope SuspendHeartBeat;
+			FDisableHitchDetectorScope SuspendGameThreadHitch;
+
+			ScriptTrace = FFrame::GetScriptCallstack(true /* bReturnEmpty */);
+			if (!ScriptTrace.IsEmpty())
+			{
+				ScriptTrace = "ScriptTrace: \n" + ScriptTrace;
+			}
+
+			// Walk the stack and dump it to the allocated memory.
+			StackTraceDepth = FPlatformStackWalk::CaptureStackBackTrace(StackTrace, MAX_DEPTH);
+			ProcessId = FPlatformProcess::GetCurrentProcessId();
+		}
+#endif // !(UE_BUILD_SHIPPING)
+
 		const uint64 WidgetId = SlateTraceDetail::GetWidgetId(Widget);
 		const uint64 InvestigatorId = SlateTraceDetail::GetWidgetIdIfValid(Investigator);
 
@@ -193,6 +234,9 @@ void FSlateTrace::WidgetInvalidated(const SWidget* Widget, const SWidget* Invest
 			<< WidgetInvalidated.Cycle(FPlatformTime::Cycles64())
 			<< WidgetInvalidated.WidgetId(WidgetId)
 			<< WidgetInvalidated.InvestigatorId(InvestigatorId)
+			<< WidgetInvalidated.ScriptTrace(*ScriptTrace)
+			<< WidgetInvalidated.Callstack(StackTrace, StackTraceDepth)
+			<< WidgetInvalidated.ProcessId(ProcessId)
 			<< WidgetInvalidated.InvalidateWidgetReason(static_cast<uint8>(Reason));
 	}
 }
