Edigrate 14642501. Fix Nonunity.

diff --git a/Engine/Plugins/Slate/SlateInsights/Source/SlateInsights/Private/SlateInsightsStyle.h b/Engine/Plugins/Slate/SlateInsights/Source/SlateInsights/Private/SlateInsightsStyle.h
--- a/Engine/Plugins/Slate/SlateInsights/Source/SlateInsights/Private/SlateInsightsStyle.h
+++ b/Engine/Plugins/Slate/SlateInsights/Source/SlateInsights/Private/SlateInsightsStyle.h
@@ -5,6 +5,7 @@
 #include "Styling/CoreStyle.h"
 #include "Styling/SlateStyle.h"
 #include "Styling/SlateStyleRegistry.h"
+#include "Styling/SlateTypes.h"
 
 namespace UE
 {
