Resolve slate insight callstacks in game process so they can be saved.

diff --git a/Engine/Plugins/Slate/SlateInsights/Source/SlateInsights/Private/SSlateFrameSchematicView.cpp b/Engine/Plugins/Slate/SlateInsights/Source/SlateInsights/Private/SSlateFrameSchematicView.cpp
--- a/Engine/Plugins/Slate/SlateInsights/Source/SlateInsights/Private/SSlateFrameSchematicView.cpp
+++ b/Engine/Plugins/Slate/SlateInsights/Source/SlateInsights/Private/SSlateFrameSchematicView.cpp
@@ -628,7 +628,7 @@ void SSlateFrameSchematicView::RefreshNodes_Invalidation(const FSlateProvider* S
 	// Build a flat list of all the invalidation
 	const FSlateProvider::TWidgetInvalidatedTimeline& InvalidatedTimeline = SlateProvider->GetWidgetInvalidatedTimeline();
 	InvalidatedTimeline.EnumerateEvents(StartTime, EndTime,
-		[&InvalidationMap](double EventStartTime, double EventEndTime, uint32 /*Depth*/, const Message::FWidgetInvalidatedMessage& Message)
+		[&InvalidationMap, &SlateProvider](double EventStartTime, double EventEndTime, uint32 /*Depth*/, const Message::FWidgetInvalidatedMessage& Message)
 		{
 			TSharedPtr<Private::FWidgetUniqueInvalidatedInfo> WidgetInfo;
 			TSharedPtr<Private::FWidgetUniqueInvalidatedInfo> InvestigatorInfo;
@@ -665,9 +665,10 @@ void SSlateFrameSchematicView::RefreshNodes_Invalidation(const FSlateProvider* S
 				WidgetInfo->ScriptTrace = Message.ScriptTrace;
 			}
 
-			if (!Message.Callstack.IsEmpty())
+			const FString* Callstack = SlateProvider->FindInvalidationCallstack(Message.SourceCycle);
+			if (Callstack)
 			{
-				WidgetInfo->Callstack = Message.Callstack;
+				WidgetInfo->Callstack = *Callstack;
 			}
 
 			return TraceServices::EEventEnumerate::Continue;
diff --git a/Engine/Plugins/Slate/SlateInsights/Source/SlateInsights/Private/SlateAnalyzer.cpp b/Engine/Plugins/Slate/SlateInsights/Source/SlateInsights/Private/SlateAnalyzer.cpp
--- a/Engine/Plugins/Slate/SlateInsights/Source/SlateInsights/Private/SlateAnalyzer.cpp
+++ b/Engine/Plugins/Slate/SlateInsights/Source/SlateInsights/Private/SlateAnalyzer.cpp
@@ -31,6 +31,7 @@ void FSlateAnalyzer::OnAnalysisBegin(const FOnAnalysisContext& Context)
 	Builder.RouteEvent(RouteId_WidgetInvalidated, "SlateTrace", "WidgetInvalidated");
 	Builder.RouteEvent(RouteId_RootInvalidated, "SlateTrace", "RootInvalidated");
 	Builder.RouteEvent(RouteId_RootChildOrderInvalidated, "SlateTrace", "RootChildOrderInvalidated");
+	Builder.RouteEvent(RouteId_InvalidationCallstack, "SlateTrace", "InvalidationCallstack");
 }
 
 bool FSlateAnalyzer::OnEvent(uint16 RouteId, EStyle Style, const FOnEventContext& Context)
@@ -98,6 +99,11 @@ bool FSlateAnalyzer::OnEvent(uint16 RouteId, EStyle Style, const FOnEventContext
 		SlateProvider.AddWidgetInvalidatedEvent(Time, Message::FWidgetInvalidatedMessage::FromChildOrder(EventData));
 		break;
 	}
+	case RouteId_InvalidationCallstack:
+	{
+		SlateProvider.ProcessInvalidationCallstack(Message::FInvalidationCallstackMessage(EventData));
+		break;
+	}
 	}
 
 	return true;
diff --git a/Engine/Plugins/Slate/SlateInsights/Source/SlateInsights/Private/SlateAnalyzer.h b/Engine/Plugins/Slate/SlateInsights/Source/SlateInsights/Private/SlateAnalyzer.h
--- a/Engine/Plugins/Slate/SlateInsights/Source/SlateInsights/Private/SlateAnalyzer.h
+++ b/Engine/Plugins/Slate/SlateInsights/Source/SlateInsights/Private/SlateAnalyzer.h
@@ -33,6 +33,7 @@ private:
 		RouteId_WidgetInvalidated,
 		RouteId_RootInvalidated,
 		RouteId_RootChildOrderInvalidated,
+		RouteId_InvalidationCallstack,
 	};
 
 	TraceServices::IAnalysisSession& Session;
diff --git a/Engine/Plugins/Slate/SlateInsights/Source/SlateInsights/Private/SlateProvider.cpp b/Engine/Plugins/Slate/SlateInsights/Source/SlateInsights/Private/SlateProvider.cpp
--- a/Engine/Plugins/Slate/SlateInsights/Source/SlateInsights/Private/SlateProvider.cpp
+++ b/Engine/Plugins/Slate/SlateInsights/Source/SlateInsights/Private/SlateProvider.cpp
@@ -4,6 +4,7 @@
 #include "HAL/PlatformStackWalk.h"
 #include "HAL/PlatformProcess.h"
 #include "Insights/ViewModels/TimingEventsTrack.h"
+#include "Trace/SlateTrace.h"
 
 #define LOCTEXT_NAMESPACE "SlateProvider"
 
@@ -31,6 +32,12 @@ namespace Message
 		static_assert(sizeof(ESlateTraceApplicationFlags) == sizeof(uint8), "ESlateTraceApplicationFlags is not a uint8");
 	}
 
+	UE::SlateInsights::Message::FInvalidationCallstackMessage::FInvalidationCallstackMessage(const UE::Trace::IAnalyzer::FEventData& EventData)
+		: SourceCycle(EventData.GetValue<uint64>("SourceCycle"))
+	{
+		EventData.GetString("CallstackText", Callstack);
+	}
+
 	FWidgetInfo::FWidgetInfo(const UE::Trace::IAnalyzer::FEventData& EventData)
 		: WidgetId(EventData.GetValue<uint64>("WidgetId"))
 		, Path()
@@ -53,62 +60,36 @@ namespace Message
 		static_assert(sizeof(EInvalidateWidgetReason) == sizeof(uint8), "EInvalidateWidgetReason is not a uint8");
 
 		FWidgetInvalidatedMessage Message;
+		Message.SourceCycle = EventData.GetValue<uint64>("Cycle");
 		Message.WidgetId = EventData.GetValue<uint64>("WidgetId");
 		Message.InvestigatorId = EventData.GetValue<uint64>("InvestigatorId");
 		Message.InvalidationReason = static_cast<EInvalidateWidgetReason>(EventData.GetValue<uint8>("InvalidateWidgetReason"));
 		EventData.GetString("ScriptTrace", Message.ScriptTrace);
-		Message.Callstack = GetCallstack(EventData);
 		return Message;
 	}
 
 	FWidgetInvalidatedMessage FWidgetInvalidatedMessage::FromRoot(const UE::Trace::IAnalyzer::FEventData& EventData)
 	{
 		FWidgetInvalidatedMessage Message;
+		Message.SourceCycle = EventData.GetValue<uint64>("Cycle");
 		Message.WidgetId = EventData.GetValue<uint64>("WidgetId");
 		Message.InvestigatorId = EventData.GetValue<uint64>("InvestigatorId");
 		Message.InvalidationReason = EInvalidateWidgetReason::Layout;
 		EventData.GetString("ScriptTrace", Message.ScriptTrace);
-		Message.Callstack = GetCallstack(EventData);
 		return Message;
 	}
 
 	FWidgetInvalidatedMessage FWidgetInvalidatedMessage::FromChildOrder(const UE::Trace::IAnalyzer::FEventData& EventData)
 	{
 		FWidgetInvalidatedMessage Message;
+		Message.SourceCycle = EventData.GetValue<uint64>("Cycle");
 		Message.WidgetId = EventData.GetValue<uint64>("WidgetId");
 		Message.InvestigatorId = EventData.GetValue<uint64>("InvestigatorId");
 		Message.InvalidationReason = EInvalidateWidgetReason::ChildOrder;
 		EventData.GetString("ScriptTrace", Message.ScriptTrace);
-		Message.Callstack = GetCallstack(EventData);
 		return Message;
 	}
 
-	FString FWidgetInvalidatedMessage::GetCallstack(const UE::Trace::IAnalyzer::FEventData& EventData)
-	{
-		FString Callstack = "";
-		const auto& CallstackReader = EventData.GetArrayView<uint64>("Callstack");
-		if (CallstackReader.Num() != 0)
-		{
-			FPlatformStackWalk::InitStackWalkingForProcess(FPlatformProcess::OpenProcess(EventData.GetValue<uint32>("ProcessId")));
-			uint8 StackDepth = 0;
-			for (uint64 ProgramCounter : CallstackReader)
-			{
-				// Skip the first two backraces, that's from us.
-				if (StackDepth++ >= 2)
-				{
-					FString SymbolAsText;
-
-					// Note: Done insight thread as this process is very slow, possibly 1~80ms based on widget complexity.
-					FProgramCounterSymbolInfoEx SymbolInfo;
-					FPlatformStackWalk::ProgramCounterToSymbolInfoEx(ProgramCounter, SymbolInfo);
-					FPlatformStackWalk::SymbolInfoToHumanReadableStringEx(SymbolInfo, SymbolAsText);
-
-					Callstack += SymbolAsText + "\r\n";
-				}
-			}
-		}
-		return Callstack;
-	}
 } //namespace Message
 
 FSlateProvider::FSlateProvider(TraceServices::IAnalysisSession& InSession)
@@ -178,6 +159,13 @@ void FSlateProvider::AddWidgetInvalidatedEvent(double Seconds, Message::FWidgetI
 	WidgetInvalidatedTimeline.EmplaceEvent(Seconds, UpdatedMessage);
 }
 
+void FSlateProvider::ProcessInvalidationCallstack(Message::FInvalidationCallstackMessage CallstackMessage)
+{
+	Session.WriteAccessCheck();
+
+	InvalidationCallstacks.Add(CallstackMessage.SourceCycle, CallstackMessage.Callstack);
+}
+
 } //namespace SlateInsights
 } //namespace UE
 
diff --git a/Engine/Plugins/Slate/SlateInsights/Source/SlateInsights/Private/SlateProvider.h b/Engine/Plugins/Slate/SlateInsights/Source/SlateInsights/Private/SlateProvider.h
--- a/Engine/Plugins/Slate/SlateInsights/Source/SlateInsights/Private/SlateProvider.h
+++ b/Engine/Plugins/Slate/SlateInsights/Source/SlateInsights/Private/SlateProvider.h
@@ -64,18 +64,17 @@ struct FWidgetUpdatedMessage
 
 struct FWidgetInvalidatedMessage
 {
+	uint64 SourceCycle;
 	FWidgetId WidgetId;
 	FWidgetId InvestigatorId;
 	EInvalidateWidgetReason InvalidationReason = EInvalidateWidgetReason::None;
 	bool bRootInvalidated = false;
 	bool bRootChildOrderInvalidated = false;
 	FString ScriptTrace;
-	FString Callstack;
 
 	static FWidgetInvalidatedMessage FromWidget(const UE::Trace::IAnalyzer::FEventData& EventData);
 	static FWidgetInvalidatedMessage FromRoot(const UE::Trace::IAnalyzer::FEventData& EventData);
 	static FWidgetInvalidatedMessage FromChildOrder(const UE::Trace::IAnalyzer::FEventData& EventData);
-	static FString GetCallstack(const UE::Trace::IAnalyzer::FEventData& EventData);
 };
 
 struct FApplicationTickedMessage
@@ -94,6 +93,14 @@ struct FApplicationTickedMessage
 	FApplicationTickedMessage(const UE::Trace::IAnalyzer::FEventData& EventData);
 };
 
+struct FInvalidationCallstackMessage
+{
+	uint64 SourceCycle;
+	FString Callstack;
+
+	FInvalidationCallstackMessage(const Trace::IAnalyzer::FEventData& EventData);
+};
+
 } //namespace Message
 
 class FSlateProvider : public TraceServices::IProvider
@@ -112,6 +119,7 @@ public:
 	void AddApplicationTickedEvent(double Seconds, Message::FApplicationTickedMessage Message);
 	void AddWidgetUpdatedEvent(double Seconds, Message::FWidgetUpdatedMessage UpdatedMessage);
 	void AddWidgetInvalidatedEvent(double Seconds, Message::FWidgetInvalidatedMessage InvalidatedMessage);
+	void ProcessInvalidationCallstack(Message::FInvalidationCallstackMessage InvalidatedMessage);
 
 	/** */
 	using TApplicationTickedTimeline = TraceServices::TPointTimeline<Message::FApplicationTickedMessage>;
@@ -170,10 +178,17 @@ public:
 		return WidgetInfos.Find(WidgetId);
 	}
 
+	/** Find a callstack that should have been saved off from the Slate Provider */
+	const FString* FindInvalidationCallstack(uint64 SourceCycle) const
+	{
+		return InvalidationCallstacks.Find(SourceCycle);
+	}
+
 private:
 	TraceServices::IAnalysisSession& Session;
 
 	TMap<Message::FWidgetId, Message::FWidgetInfo> WidgetInfos;
+	TMap<uint64, FString> InvalidationCallstacks;
 
 	TWidgetTimeline WidgetTimelines;
 	TApplicationTickedTimeline ApplicationTickedTimeline;
diff --git a/Engine/Source/Runtime/SlateCore/Private/Trace/SlateTrace.cpp b/Engine/Source/Runtime/SlateCore/Private/Trace/SlateTrace.cpp
--- a/Engine/Source/Runtime/SlateCore/Private/Trace/SlateTrace.cpp
+++ b/Engine/Source/Runtime/SlateCore/Private/Trace/SlateTrace.cpp
@@ -1,6 +1,7 @@
 // Copyright Epic Games, Inc. All Rights Reserved.
 
 #include "Trace/SlateTrace.h"
+#include "Async/Async.h"
 
 #if UE_SLATE_TRACE_ENABLED
 
@@ -73,7 +74,6 @@ UE_TRACE_EVENT_BEGIN(SlateTrace, WidgetInvalidated)
 	UE_TRACE_EVENT_FIELD(uint8, InvalidateWidgetReason)		// The reason of the invalidation. (EInvalidateWidgetReason)
 	UE_TRACE_EVENT_FIELD(UE::Trace::WideString, ScriptTrace)// Optional script trace for root widget invalidations
 	UE_TRACE_EVENT_FIELD(uint64[], Callstack)				// Optional callstack for root widget invalidations
-	UE_TRACE_EVENT_FIELD(uint32, ProcessId)					// Optional proccess ID where the invalidation occureed.
 UE_TRACE_EVENT_END()
 
 UE_TRACE_EVENT_BEGIN(SlateTrace, RootInvalidated)
@@ -88,6 +88,12 @@ UE_TRACE_EVENT_BEGIN(SlateTrace, RootChildOrderInvalidated)
 	UE_TRACE_EVENT_FIELD(uint64, InvestigatorId)		// Widget unique ID that investigated the invalidation.
 UE_TRACE_EVENT_END()
 
+UE_TRACE_EVENT_BEGIN(SlateTrace, InvalidationCallstack)
+	UE_TRACE_EVENT_FIELD(uint64, SourceCycle)              // Cycle during which a callstack was captured
+	UE_TRACE_EVENT_FIELD(UE::Trace::AnsiString, CallstackText) // Text of the captured callstack
+UE_TRACE_EVENT_END()
+
+
 //-----------------------------------------------------------------------------------//
 
 namespace SlateTraceDetail
@@ -229,15 +235,51 @@ void FSlateTrace::WidgetInvalidated(const SWidget* Widget, const SWidget* Invest
 
 		const uint64 WidgetId = SlateTraceDetail::GetWidgetId(Widget);
 		const uint64 InvestigatorId = SlateTraceDetail::GetWidgetIdIfValid(Investigator);
+		const uint64 Cycle = FPlatformTime::Cycles64();
 
 		UE_TRACE_LOG(SlateTrace, WidgetInvalidated, SlateChannel)
-			<< WidgetInvalidated.Cycle(FPlatformTime::Cycles64())
+			<< WidgetInvalidated.Cycle(Cycle)
 			<< WidgetInvalidated.WidgetId(WidgetId)
 			<< WidgetInvalidated.InvestigatorId(InvestigatorId)
 			<< WidgetInvalidated.ScriptTrace(*ScriptTrace)
 			<< WidgetInvalidated.Callstack(StackTrace, StackTraceDepth)
-			<< WidgetInvalidated.ProcessId(ProcessId)
 			<< WidgetInvalidated.InvalidateWidgetReason(static_cast<uint8>(Reason));
+
+#if !(UE_BUILD_SHIPPING)
+		if (!Investigator && bCaptureRootInvalidationCallstacks)
+		{
+			// Note: Done in seperate thread as symbol resolution is very slow, possibly 1~80ms based on widget complexity.
+			AsyncTask(ENamedThreads::AnyBackgroundThreadNormalTask, [Cycle, StackTrace, StackTraceDepth]()
+			{
+				const SIZE_T CallStackSize = 65535;
+				ANSICHAR* CallStack = (ANSICHAR*)FMemory::SystemMalloc(CallStackSize);
+				if (CallStack != nullptr)
+				{
+					CallStack[0] = 0;
+				}
+
+				uint32 StackDepth = 0;
+				while (StackDepth < StackTraceDepth)
+				{
+					// Skip the first two backraces, that's from us.
+					if (StackDepth >= 2 && StackTrace[StackDepth])
+					{
+						FProgramCounterSymbolInfo SymbolInfo;
+						FPlatformStackWalk::ProgramCounterToSymbolInfo(StackTrace[StackDepth], SymbolInfo);
+						FPlatformStackWalk::SymbolInfoToHumanReadableString(SymbolInfo, CallStack, CallStackSize);
+						FCStringAnsi::Strncat(CallStack, LINE_TERMINATOR_ANSI, CallStackSize);
+					}
+					StackDepth++;
+				}
+
+				UE_TRACE_LOG(SlateTrace, InvalidationCallstack, SlateChannel)
+					<< InvalidationCallstack.SourceCycle(Cycle)
+					<< InvalidationCallstack.CallstackText(CallStack);
+
+				FMemory::SystemFree(CallStack);
+			});
+		}
+#endif // !(UE_BUILD_SHIPPING)
 	}
 }
 
