SlateInsights: Add a Widget to find and display a Swidget info.

diff --git a/Engine/Plugins/Slate/SlateInsights/Source/SlateInsights/Private/SSlateFrameSchematicView.cpp b/Engine/Plugins/Slate/SlateInsights/Source/SlateInsights/Private/SSlateFrameSchematicView.cpp
--- a/Engine/Plugins/Slate/SlateInsights/Source/SlateInsights/Private/SSlateFrameSchematicView.cpp
+++ b/Engine/Plugins/Slate/SlateInsights/Source/SlateInsights/Private/SSlateFrameSchematicView.cpp
@@ -5,15 +5,20 @@
 #include "SSlateTraceFlags.h"
 
 #include "TraceServices/Model/AnalysisSession.h"
+#include "Insights/Common/TimeUtils.h"
 #include "Insights/ITimingViewSession.h"
 #include "Insights/ViewModels/ITimingEvent.h"
 #include "Widgets/InvalidateWidgetReason.h"
 #include "Widgets/SBoxPanel.h"
-#include "Widgets/SNullWidget.h"
+#include "Widgets/Input/NumericTypeInterface.h"
+#include "Widgets/Input/SMultiLineEditableTextBox.h"
+#include "Widgets/Input/SSearchBox.h"
+#include "Widgets/Layout/SExpandableArea.h"
+#include "Widgets/Layout/SGridPanel.h"
 #include "Widgets/Layout/SHeader.h"
+#include "Widgets/SNullWidget.h"
 #include "Widgets/Text/STextBlock.h"
 #include "Widgets/Views/STableRow.h"
-#include "Widgets/Input/SMultiLineEditableTextBox.h"
 
 #include "Framework/Commands/UIAction.h"
 #include "Framework/MultiBox/MultiBoxBuilder.h"
@@ -39,6 +44,7 @@ namespace Private
 	const FName ColumnNumber("Number");
 	const FName ColumnFlag("Flag");
 
+	/** */
 	struct FWidgetUniqueInvalidatedInfo
 	{
 		FWidgetUniqueInvalidatedInfo(Message::FWidgetId InWidgetId, EInvalidateWidgetReason InReason)
@@ -54,6 +60,7 @@ namespace Private
 		FString Callstack;
 	};
 
+	/** */
 	struct FWidgetUpdateInfo
 	{
 		FWidgetUpdateInfo(Message::FWidgetId InWidgetId, EWidgetUpdateFlags InUpdateFlags)
@@ -64,6 +71,7 @@ namespace Private
 		uint32 Count;
 	};
 
+	/** */
 	class FWidgetUniqueInvalidatedInfoRow : public SMultiColumnTableRow<TSharedPtr<FWidgetUniqueInvalidatedInfo>>
 	{
 	public:
@@ -129,6 +137,7 @@ namespace Private
 		}
 	};
 
+	/** */
 	struct FWidgetUpdateInfoRow : public SMultiColumnTableRow<TSharedPtr<FWidgetUpdateInfo>>
 	{
 		SLATE_BEGIN_ARGS(FWidgetUpdateInfoRow) {}
@@ -186,6 +195,142 @@ namespace Private
 		}
 		return FText::GetEmpty();
 	}
+
+	/** */
+	class SSlateWidgetSearch : public SCompoundWidget
+	{
+	public:
+		SLATE_BEGIN_ARGS(SSlateWidgetSearch) {}
+		SLATE_END_ARGS()
+
+		void Construct(const FArguments& InArgs)
+		{
+			NumericInterface = MakeUnique<TDefaultNumericTypeInterface<uint64>>();
+
+			TSharedRef<SGridPanel> GridPanel = SNew(SGridPanel)
+				.FillColumn(1, 1.0f);
+			{
+				int32 SlotCount = 0;
+				auto BuildLabelAndValue = [GridPanel, &SlotCount](const FText& Label, TSharedPtr<STextBlock>& ValueWidget)
+				{
+					GridPanel->AddSlot(0, SlotCount)
+						[
+							SNew(STextBlock)
+							.Text(Label)
+						];
+
+					GridPanel->AddSlot(1, SlotCount)
+						.Padding(5.0f, 0.0f, 0.0f, 0.0f)
+						[
+							SAssignNew(ValueWidget, STextBlock)
+						];
+
+					++SlotCount;
+				};
+
+				BuildLabelAndValue(LOCTEXT("WidgetID", "Widget ID"), WidgetIdWidget);
+				BuildLabelAndValue(LOCTEXT("Path", "Path"), PathWidget);
+				BuildLabelAndValue(LOCTEXT("DebugInfo", "Debug Info"), DebugInfoWidget);
+				BuildLabelAndValue(LOCTEXT("CreatedTime", "Created Time"), CreatedTimeWidget);
+				BuildLabelAndValue(LOCTEXT("DestroyedTime", "Destroyed Time"), DestroyedTimeWidget);
+			}
+
+			ChildSlot
+			[
+				SNew(SVerticalBox)
+
+				+ SVerticalBox::Slot()
+				.AutoHeight()
+				.Padding(FMargin(0.f, 4.f))
+				[
+					SAssignNew(SearchBoxWidget, SSearchBox)
+					.HintText(LOCTEXT("WidgetSearchBoxHint", "Search Widget ID"))
+					.OnTextChanged(this, &SSlateWidgetSearch::HandleSearchBox_OnTextChanged)
+					.IsEnabled(this, &SSlateWidgetSearch::HandleSearchBox_IsEnabled)
+					.ToolTipText(LOCTEXT("FilterSearchHint", "Type a Widget ID here to search for the widget"))
+				]
+
+				+ SVerticalBox::Slot()
+				.AutoHeight()
+				[
+					GridPanel
+				]
+			];
+		}
+
+		void SetSession(const TraceServices::IAnalysisSession* InAnalysisSession)
+		{
+			AnalysisSession = InAnalysisSession;
+			Search();
+		}
+
+		bool HandleSearchBox_IsEnabled() const
+		{
+			return AnalysisSession != nullptr;
+		}
+
+		void HandleSearchBox_OnTextChanged(const FText& InText)
+		{
+			TOptional<uint64> Result = NumericInterface->FromString(InText.ToString(), WidgetId);
+			if (Result.IsSet())
+			{
+				WidgetId = Result.GetValue();
+				Search();
+				SearchBoxWidget->SetError(FText::GetEmpty());
+			}
+			else
+			{
+				SearchBoxWidget->SetError(LOCTEXT("NotAValidId", "Not a valid Widget Id. Widget Ids are numbers."));
+			}
+		}
+
+		void Search()
+		{
+			bool bClearText = true;
+			if (AnalysisSession && PathWidget && WidgetId != 0)
+			{
+				const FSlateProvider* SlateProvider = AnalysisSession->ReadProvider<FSlateProvider>(FSlateProvider::ProviderName);
+				if (SlateProvider)
+				{
+					TraceServices::FAnalysisSessionReadScope SessionReadScope(*AnalysisSession);
+
+					if (const Message::FWidgetInfo* WidgetInfo = SlateProvider->FindWidget(WidgetId))
+					{
+						bClearText = false;
+
+						WidgetIdWidget->SetText(FText::FromString(NumericInterface->ToString(WidgetInfo->WidgetId.GetValue())));
+						PathWidget->SetText(FText::FromString(WidgetInfo->Path));
+						DebugInfoWidget->SetText(FText::FromString(WidgetInfo->DebugInfo));
+
+						const double StartTime = SlateProvider->GetWidgetTimeline().GetEventStartTime(WidgetInfo->EventIndex);
+						CreatedTimeWidget->SetText(FText::FromString(TimeUtils::FormatTime(StartTime, TimeUtils::Milisecond)));
+						const double EndTime = SlateProvider->GetWidgetTimeline().GetEventEndTime(WidgetInfo->EventIndex);
+						DestroyedTimeWidget->SetText(FText::FromString(TimeUtils::FormatTime(EndTime, TimeUtils::Milisecond)));
+					}
+				}
+			}
+			
+			if (bClearText)
+			{
+				WidgetIdWidget->SetText(FText::GetEmpty());
+				PathWidget->SetText(FText::GetEmpty());
+				DebugInfoWidget->SetText(FText::GetEmpty());
+				CreatedTimeWidget->SetText(FText::GetEmpty());
+				DestroyedTimeWidget->SetText(FText::GetEmpty());
+			}
+		}
+
+		const TraceServices::IAnalysisSession* AnalysisSession = nullptr;
+		uint64 WidgetId = 0;
+
+		TSharedPtr<SSearchBox> SearchBoxWidget;
+		TSharedPtr<STextBlock> WidgetIdWidget;
+		TSharedPtr<STextBlock> PathWidget;
+		TSharedPtr<STextBlock> DebugInfoWidget;
+		TSharedPtr<STextBlock> CreatedTimeWidget;
+		TSharedPtr<STextBlock> DestroyedTimeWidget;
+		TUniquePtr<INumericTypeInterface<uint64>> NumericInterface;
+	};
 } //namespace Private
 
 void SSlateFrameSchematicView::Construct(const FArguments& InArgs)
@@ -197,6 +342,23 @@ void SSlateFrameSchematicView::Construct(const FArguments& InArgs)
 	[
 		SNew(SVerticalBox)
 
+		+ SVerticalBox::Slot()
+		.AutoHeight()
+		.Padding(FMargin(2.0f))
+		[
+			SNew(SExpandableArea)
+			.InitiallyCollapsed(true)
+			.HeaderContent()
+			[
+				SNew(STextBlock)
+				.Text(LOCTEXT("SearchWidget", "Search Widget"))
+			]
+			.BodyContent()
+			[
+				SAssignNew(WidgetSearchBox, Private::SSlateWidgetSearch)
+			]
+		]
+		
 		+ SVerticalBox::Slot()
 		.AutoHeight()
 		.Padding(FMargin(2.0f))
@@ -309,6 +471,8 @@ void SSlateFrameSchematicView::Construct(const FArguments& InArgs)
 		]
 	];
 
+	WidgetSearchBox->SetSession(AnalysisSession);
+
 	RefreshNodes();
 }
 
@@ -341,6 +505,11 @@ void SSlateFrameSchematicView::SetSession(Insights::ITimingViewSession* InTiming
 		TimingViewSession->OnSelectedEventChanged().AddSP(this, &SSlateFrameSchematicView::HandleSelectionEventChanged);
 	}
 
+	if (WidgetSearchBox)
+	{
+		WidgetSearchBox->SetSession(InAnalysisSession);
+	}
+
 	RefreshNodes();
 }
 
diff --git a/Engine/Plugins/Slate/SlateInsights/Source/SlateInsights/Private/SSlateFrameSchematicView.h b/Engine/Plugins/Slate/SlateInsights/Source/SlateInsights/Private/SSlateFrameSchematicView.h
--- a/Engine/Plugins/Slate/SlateInsights/Source/SlateInsights/Private/SSlateFrameSchematicView.h
+++ b/Engine/Plugins/Slate/SlateInsights/Source/SlateInsights/Private/SSlateFrameSchematicView.h
@@ -27,6 +27,7 @@ namespace SlateInsights
 class FSlateProvider;
 namespace Private { struct FWidgetUniqueInvalidatedInfo; }
 namespace Private { struct FWidgetUpdateInfo; }
+namespace Private { class SSlateWidgetSearch; }
 
 class SSlateFrameSchematicView : public SCompoundWidget
 {
@@ -69,6 +70,7 @@ private:
 	TSharedPtr<SListView<TSharedPtr<Private::FWidgetUpdateInfo>>> WidgetUpdateInfoListView;
 	TArray<TSharedPtr<Private::FWidgetUpdateInfo>> WidgetUpdateInfos;
 
+	TSharedPtr<Private::SSlateWidgetSearch> WidgetSearchBox;
 	TSharedPtr<STextBlock> InvalidationSummary;
 	TSharedPtr<STextBlock> UpdateSummary;
 	TSharedPtr<SMultiLineEditableTextBox> ScriptAndCallStackTextBox;
diff --git a/Engine/Plugins/Slate/SlateInsights/Source/SlateInsights/Private/SlateProvider.cpp b/Engine/Plugins/Slate/SlateInsights/Source/SlateInsights/Private/SlateProvider.cpp
--- a/Engine/Plugins/Slate/SlateInsights/Source/SlateInsights/Private/SlateProvider.cpp
+++ b/Engine/Plugins/Slate/SlateInsights/Source/SlateInsights/Private/SlateProvider.cpp
@@ -116,8 +116,8 @@ void FSlateProvider::AddWidget(double Seconds, uint64 WidgetId)
 
 	Message::FWidgetInfo Info;
 	Info.WidgetId = WidgetId;
-	Message::FWidgetInfo& InfoRef = WidgetInfos.Emplace(Info.WidgetId, MoveTemp(Info));
-	InfoRef.EventIndex = WidgetTimelines.EmplaceBeginEvent(Seconds, Info.WidgetId);
+	Info.EventIndex = WidgetTimelines.EmplaceBeginEvent(Seconds, WidgetId);
+	WidgetInfos.Emplace(Info.WidgetId, MoveTemp(Info));
 }
 
 void FSlateProvider::SetWidgetInfo(double Seconds, Message::FWidgetInfo Info)
@@ -126,12 +126,13 @@ void FSlateProvider::SetWidgetInfo(double Seconds, Message::FWidgetInfo Info)
 
 	if (Message::FWidgetInfo* FoundInfo = WidgetInfos.Find(Info.WidgetId))
 	{
+		Info.EventIndex = FoundInfo->EventIndex;
 		*FoundInfo = MoveTemp(Info);
 	}
 	else
 	{
-		Message::FWidgetInfo& InfoRef = WidgetInfos.Emplace(Info.WidgetId, MoveTemp(Info));
-		InfoRef.EventIndex = WidgetTimelines.EmplaceBeginEvent(Seconds, Info.WidgetId);
+		Info.EventIndex = WidgetTimelines.EmplaceBeginEvent(Seconds, Info.WidgetId);
+		WidgetInfos.Emplace(Info.WidgetId, MoveTemp(Info));
 	}
 }
 
diff --git a/Engine/Plugins/Slate/SlateInsights/Source/SlateInsights/Private/SlateProvider.h b/Engine/Plugins/Slate/SlateInsights/Source/SlateInsights/Private/SlateProvider.h
--- a/Engine/Plugins/Slate/SlateInsights/Source/SlateInsights/Private/SlateProvider.h
+++ b/Engine/Plugins/Slate/SlateInsights/Source/SlateInsights/Private/SlateProvider.h
@@ -36,6 +36,7 @@ public:
 	template<typename T, typename U = typename TEnableIf<TIsSame<T, uint64>::Value>::Type>
 	constexpr FWidgetId(T InValue) : Value(InValue) {}
 	explicit operator bool() const { return Value != 0; }
+	uint64 GetValue() const { return Value; }
 	friend uint32 GetTypeHash(const FWidgetId& Key) { return ::GetTypeHash(Key.Value); }
 	friend bool operator==(const FWidgetId A, const FWidgetId B) { return A.Value == B.Value; }
 	friend bool operator!=(const FWidgetId A, const FWidgetId B) { return A.Value != B.Value; }
diff --git a/Engine/Source/Developer/TraceServices/Public/Model/IntervalTimeline.h b/Engine/Source/Developer/TraceServices/Public/Model/IntervalTimeline.h
--- a/Engine/Source/Developer/TraceServices/Public/Model/IntervalTimeline.h
+++ b/Engine/Source/Developer/TraceServices/Public/Model/IntervalTimeline.h
@@ -32,6 +32,8 @@ public:
 	virtual uint64 GetModCount() const override { return ModCount; }
 	virtual uint64 GetEventCount() const override { return Events.Num(); }
 	virtual const EventType& GetEvent(uint64 InIndex) const override { return Events[InIndex].Event; }
+	double GetEventStartTime(uint64 InIndex) const { return Events[InIndex].StartTime; }
+	double GetEventEndTime(uint64 InIndex) const { return Events[InIndex].EndTime; }
 	virtual double GetStartTime() const override { return Events.Num() > 0 ? Events[0].StartTime : 0.0; }
 	virtual double GetEndTime() const override { return Events.Num() > 0 ? Events[Events.Num() - 1].EndTime : 0.0; }
 	virtual void EnumerateEventsDownSampled(double IntervalStart, double IntervalEnd, double Resolution, typename ITimeline<EventType>::EventCallback Callback) const override { check(false); }
