SlateInsight: Add timing to Widget Update

diff --git a/Engine/Plugins/Slate/SlateInsights/Source/SlateInsights/Private/SSlateFrameSchematicView.cpp b/Engine/Plugins/Slate/SlateInsights/Source/SlateInsights/Private/SSlateFrameSchematicView.cpp
--- a/Engine/Plugins/Slate/SlateInsights/Source/SlateInsights/Private/SSlateFrameSchematicView.cpp
+++ b/Engine/Plugins/Slate/SlateInsights/Source/SlateInsights/Private/SSlateFrameSchematicView.cpp
@@ -43,6 +43,8 @@ namespace Private
 {
 	const FName ColumnWidgetId("WidgetId");
 	const FName ColumnNumber("Number");
+	const FName ColumnAffectedCount("AffectedCount");
+	const FName ColumnDuration("Duration");
 	const FName ColumnFlag("Flag");
 
 	/** */
@@ -64,12 +66,14 @@ namespace Private
 	/** */
 	struct FWidgetUpdateInfo
 	{
-		FWidgetUpdateInfo(Message::FWidgetId InWidgetId, EWidgetUpdateFlags InUpdateFlags)
-			: WidgetId(InWidgetId), UpdateFlags(InUpdateFlags), Count(1)
+		FWidgetUpdateInfo(Message::FWidgetId InWidgetId, int32 InAffectedCount, double InDuration, EWidgetUpdateFlags InUpdateFlags)
+			: WidgetId(InWidgetId), AffectedCount(InAffectedCount), Count(1), Duration(InDuration), UpdateFlags(InUpdateFlags)
 		{ }
 		Message::FWidgetId WidgetId;
-		EWidgetUpdateFlags UpdateFlags;
+		int32 AffectedCount;
 		uint32 Count;
+		double Duration;
+		EWidgetUpdateFlags UpdateFlags;
 	};
 
 	/** */
@@ -169,6 +173,16 @@ namespace Private
 				return SNew(STextBlock)
 					.Text(FText::AsNumber(Info->Count));
 			}
+			else if (Column == ColumnAffectedCount)
+			{
+				return SNew(STextBlock)
+					.Text(FText::AsNumber(Info->AffectedCount));
+			}
+			else if (Column == ColumnDuration)
+			{
+				return SNew(STextBlock)
+					.Text(FText::FromString(TimeUtils::FormatTimeAuto(Info->Duration)));
+			}
 			else if (Column == ColumnFlag)
 			{
 				return SNew(SSlateTraceWidgetUpdateFlags)
@@ -339,6 +353,9 @@ void SSlateFrameSchematicView::Construct(const FArguments& InArgs)
 	StartTime = -1.0;
 	EndTime = -1.0;
 
+	WidgetUpdateSortColumn = FName();
+	bWidgetUpdateSortAscending = false;
+
 	ChildSlot
 	[
 		SNew(SVerticalBox)
@@ -392,6 +409,7 @@ void SSlateFrameSchematicView::Construct(const FArguments& InArgs)
 					.TreeItemsSource(&WidgetInvalidationInfos)
 					.OnGenerateRow(this, &SSlateFrameSchematicView::HandleUniqueInvalidatedMakeTreeRowWidget)
 					.OnGetChildren(this, &SSlateFrameSchematicView::HandleUniqueInvalidatedChildrenForInfo)
+					.OnItemToString_Debug(this, &SSlateFrameSchematicView::HandleWidgetInvalidateListToStringDebug)
 					.OnContextMenuOpening(this, &SSlateFrameSchematicView::HandleWidgetInvalidateListContextMenu)
 					.SelectionMode(ESelectionMode::Single)
 					.HeaderRow
@@ -401,6 +419,7 @@ void SSlateFrameSchematicView::Construct(const FArguments& InArgs)
 						+ SHeaderRow::Column(Private::ColumnWidgetId)
 						.DefaultLabel(LOCTEXT("WidgetColumn", "Widget"))
 						.FillWidth(1.f)
+						.HAlignCell(EHorizontalAlignment::HAlign_Left)
 
 						+ SHeaderRow::Column(Private::ColumnNumber)
 						.DefaultLabel(LOCTEXT("AmountColumn", "Amount"))
@@ -409,7 +428,7 @@ void SSlateFrameSchematicView::Construct(const FArguments& InArgs)
 						
 						+ SHeaderRow::Column(Private::ColumnFlag)
 						.DefaultLabel(LOCTEXT("ReasonColumn", "Reason"))
-						.FixedWidth(75.f)
+						.FixedWidth(95.f)
 						.HAlignCell(EHorizontalAlignment::HAlign_Right)
 					)
 				]
@@ -455,8 +474,25 @@ void SSlateFrameSchematicView::Construct(const FArguments& InArgs)
 
 						+ SHeaderRow::Column(Private::ColumnWidgetId)
 						.DefaultLabel(LOCTEXT("WidgetColumn", "Widget"))
-						.FillWidth(1.f)
+						.HAlignCell(EHorizontalAlignment::HAlign_Left)
+						.FillWidth(0.8f)
+						.SortMode(this, &SSlateFrameSchematicView::HandleWidgetUpdateGetSortMode, Private::ColumnWidgetId)
+						.OnSort(this, &SSlateFrameSchematicView::HandleWidgetUpdateInfoSort)
 
+						+ SHeaderRow::Column(Private::ColumnAffectedCount)
+						.DefaultLabel(LOCTEXT("AffectedColumn", "Affected"))
+						.FixedWidth(50.f)
+						.HAlignCell(EHorizontalAlignment::HAlign_Right)
+						.SortMode(this, &SSlateFrameSchematicView::HandleWidgetUpdateGetSortMode, Private::ColumnAffectedCount)
+						.OnSort(this, &SSlateFrameSchematicView::HandleWidgetUpdateInfoSort)
+						
+						+ SHeaderRow::Column(Private::ColumnDuration)
+						.DefaultLabel(LOCTEXT("Duration", "Length"))
+						.FillWidth(0.2f)
+						.HAlignCell(EHorizontalAlignment::HAlign_Right)
+						.SortMode(this, &SSlateFrameSchematicView::HandleWidgetUpdateGetSortMode, Private::ColumnDuration)
+						.OnSort(this, &SSlateFrameSchematicView::HandleWidgetUpdateInfoSort)
+						
 						+ SHeaderRow::Column(Private::ColumnNumber)
 						.DefaultLabel(LOCTEXT("AmountColumn", "Amount"))
 						.FixedWidth(50.f)
@@ -464,7 +500,7 @@ void SSlateFrameSchematicView::Construct(const FArguments& InArgs)
 
 						+ SHeaderRow::Column(Private::ColumnFlag)
 						.DefaultLabel(LOCTEXT("UpdateFlagColumn", "Update"))
-						.FixedWidth(75.f)
+						.FixedWidth(95.f)
 						.HAlignCell(EHorizontalAlignment::HAlign_Right)
 					)
 				]
@@ -535,6 +571,27 @@ TSharedRef<ITableRow> SSlateFrameSchematicView::HandleWidgetUpdateInfoGenerateWi
 		Private::GetWidgetName(AnalysisSession, Item->WidgetId));
 }
 
+void SSlateFrameSchematicView::HandleWidgetUpdateInfoSort(EColumnSortPriority::Type, const FName& ColumnId, EColumnSortMode::Type SortMode)
+{
+	if (WidgetUpdateSortColumn == ColumnId)
+	{
+		bWidgetUpdateSortAscending = !bWidgetUpdateSortAscending;
+	}
+	else
+	{
+		bWidgetUpdateSortAscending = true;
+	}
+	WidgetUpdateSortColumn = ColumnId;
+	SortWidgetUpdateInfos();
+}
+
+EColumnSortMode::Type SSlateFrameSchematicView::HandleWidgetUpdateGetSortMode(FName ColumnId) const
+{
+	return WidgetUpdateSortColumn == ColumnId
+		? (bWidgetUpdateSortAscending? EColumnSortMode::Ascending : EColumnSortMode::Descending)
+		: EColumnSortMode::None;
+}
+
 void SSlateFrameSchematicView::HandleTimeMarkerChanged(Insights::ETimeChangedFlags InFlags, double InTimeMarker)
 {
 	if (!FMath::IsNearlyEqual(StartTime, InTimeMarker) && !FMath::IsNearlyEqual(EndTime, InTimeMarker))
@@ -599,6 +656,11 @@ TSharedPtr<SWidget> SSlateFrameSchematicView::HandleWidgetInvalidateListContextM
 	return MenuBuilder.MakeWidget();
 }
 
+FString SSlateFrameSchematicView::HandleWidgetInvalidateListToStringDebug(TSharedPtr<Private::FWidgetUniqueInvalidatedInfo> InInfo)
+{
+	return Private::GetWidgetName(AnalysisSession, InInfo->WidgetId).ToString();
+}
+
 bool SSlateFrameSchematicView::CanWidgetInvalidateListGotoRootWidget()
 {
 	return WidgetInvalidateInfoListView->GetSelectedItems().Num() == 1;
@@ -827,7 +889,7 @@ void SSlateFrameSchematicView::RefreshNodes_Invalidation(const FSlateProvider* S
 					InvalidationMap.Add(Message.InvestigatorId, InvestigatorInfo);
 				}
 				InvestigatorInfo->bRoot = false;
-				WidgetInfo->Investigators.Add(InvestigatorInfo);
+				WidgetInfo->Investigators.AddUnique(InvestigatorInfo);
 			}
 
 			if (!Message.ScriptTrace.IsEmpty())
@@ -864,22 +926,87 @@ void SSlateFrameSchematicView::RefreshNodes_Update(const FSlateProvider* SlatePr
 		{
 			if (TSharedPtr<Private::FWidgetUpdateInfo>* Info = WidgetUpdateInfosMap.Find(Message.WidgetId))
 			{
-				++((*Info)->Count);
-				(*Info)->UpdateFlags |= Message.UpdateFlags;
+				Private::FWidgetUpdateInfo& UpdateInfo = *(*Info);
+				// if we have the same flag again, then we have a duplicate
+				if ((UpdateInfo.UpdateFlags & Message.UpdateFlags) != EWidgetUpdateFlags::None)
+				{
+					++UpdateInfo.Count;
+				}
+				UpdateInfo.AffectedCount = FMath::Max(UpdateInfo.AffectedCount, Message.AffectedCount);
+				// If the new update is the paint message, take it's time since it include the tick and active timer
+				if (EnumHasAnyFlags(UpdateInfo.UpdateFlags, EWidgetUpdateFlags::NeedsRepaint|EWidgetUpdateFlags::NeedsVolatilePaint))
+				{
+					UpdateInfo.Duration = Message.Duration;
+				}
+				UpdateInfo.UpdateFlags |= Message.UpdateFlags;
 			}
 			else
 			{
 				WidgetUpdateInfosMap.Add(
 					Message.WidgetId,
-					MakeShared<Private::FWidgetUpdateInfo>(Message.WidgetId, Message.UpdateFlags));
+					MakeShared<Private::FWidgetUpdateInfo>(Message.WidgetId, Message.AffectedCount, Message.Duration, Message.UpdateFlags));
 			}
 			return TraceServices::EEventEnumerate::Continue;
 		});
 	WidgetUpdateInfosMap.GenerateValueArray(WidgetUpdateInfos);
+	SortWidgetUpdateInfos();
 
 	UpdateSummary->SetText(FText::Format(LOCTEXT("UpdateSummary_Formated", "{0} widgets updated."), FText::AsNumber(WidgetUpdateInfos.Num())));
 }
 
+void SSlateFrameSchematicView::SortWidgetUpdateInfos()
+{
+	if (bWidgetUpdateSortAscending)
+	{
+		if (WidgetUpdateSortColumn == Private::ColumnAffectedCount)
+		{
+			WidgetUpdateInfos.Sort([](const TSharedPtr<Private::FWidgetUpdateInfo>& A, const TSharedPtr<Private::FWidgetUpdateInfo>& B)
+				{
+					return A->AffectedCount < B->AffectedCount;
+				});
+		}
+		else if (WidgetUpdateSortColumn == Private::ColumnDuration)
+		{
+			WidgetUpdateInfos.Sort([](const TSharedPtr<Private::FWidgetUpdateInfo>& A, const TSharedPtr<Private::FWidgetUpdateInfo>& B)
+				{
+					return A->Duration < B->Duration;
+				});
+		}
+		else if (WidgetUpdateSortColumn == Private::ColumnWidgetId)
+		{
+			WidgetUpdateInfos.Sort([](const TSharedPtr<Private::FWidgetUpdateInfo>& A, const TSharedPtr<Private::FWidgetUpdateInfo>& B)
+				{
+					return A->WidgetId.GetValue() < B->WidgetId.GetValue();
+				});
+		}
+	}
+	else
+	{
+		if (WidgetUpdateSortColumn == Private::ColumnAffectedCount)
+		{
+			WidgetUpdateInfos.Sort([](const TSharedPtr<Private::FWidgetUpdateInfo>& A, const TSharedPtr<Private::FWidgetUpdateInfo>& B)
+				{
+					return A->AffectedCount > B->AffectedCount;
+				});
+		}
+		else if (WidgetUpdateSortColumn == Private::ColumnDuration)
+		{
+			WidgetUpdateInfos.Sort([](const TSharedPtr<Private::FWidgetUpdateInfo>& A, const TSharedPtr<Private::FWidgetUpdateInfo>& B)
+				{
+					return A->Duration > B->Duration;
+				});
+		}
+		else if (WidgetUpdateSortColumn == Private::ColumnWidgetId)
+		{
+			WidgetUpdateInfos.Sort([](const TSharedPtr<Private::FWidgetUpdateInfo>& A, const TSharedPtr<Private::FWidgetUpdateInfo>& B)
+				{
+					return A->WidgetId.GetValue() > B->WidgetId.GetValue();
+				});
+		}
+	}
+	WidgetUpdateInfoListView->RebuildList();
+}
+
 } //namespace SlateInsights
 } //namespace UE
 
diff --git a/Engine/Plugins/Slate/SlateInsights/Source/SlateInsights/Private/SSlateFrameSchematicView.h b/Engine/Plugins/Slate/SlateInsights/Source/SlateInsights/Private/SSlateFrameSchematicView.h
--- a/Engine/Plugins/Slate/SlateInsights/Source/SlateInsights/Private/SSlateFrameSchematicView.h
+++ b/Engine/Plugins/Slate/SlateInsights/Source/SlateInsights/Private/SSlateFrameSchematicView.h
@@ -49,10 +49,13 @@ private:
 	void HandleSelectionEventChanged(const TSharedPtr<const ITimingEvent> InEvent);
 
 	TSharedPtr<SWidget> HandleWidgetInvalidateListContextMenu();
+	FString HandleWidgetInvalidateListToStringDebug(TSharedPtr<Private::FWidgetUniqueInvalidatedInfo> Item);
 	bool CanWidgetInvalidateListGotoRootWidget();
 	void HandleWidgetInvalidateListGotoRootWidget();
 	bool CanWidgetInvalidateListViewScriptAndCallStack();
 	void HandleWidgetInvalidateListViewScriptAndCallStack();
+	void HandleWidgetUpdateInfoSort(EColumnSortPriority::Type, const FName& ColumnId, EColumnSortMode::Type SortMode);
+	EColumnSortMode::Type HandleWidgetUpdateGetSortMode(FName ColumnId) const;
 
 	virtual FReply OnMouseButtonDoubleClick(const FGeometry& InMyGeometry, const FPointerEvent& InMouseEvent) override;
 
@@ -60,6 +63,8 @@ private:
 	void RefreshNodes_Invalidation(const FSlateProvider* SlateProvider);
 	void RefreshNodes_Update(const FSlateProvider* SlateProvider);
 
+	void SortWidgetUpdateInfos();
+
 private:
 	const TraceServices::IAnalysisSession* AnalysisSession;
 	Insights::ITimingViewSession* TimingViewSession;
@@ -75,6 +80,9 @@ private:
 	TSharedPtr<STextBlock> UpdateSummary;
 	TSharedPtr<SMultiLineEditableTextBox> ScriptAndCallStackTextBox;
 
+	FName WidgetUpdateSortColumn;
+	bool bWidgetUpdateSortAscending;
+
 	double StartTime;
 	double EndTime;
 };
diff --git a/Engine/Plugins/Slate/SlateInsights/Source/SlateInsights/Private/SSlateTraceFlags.cpp b/Engine/Plugins/Slate/SlateInsights/Source/SlateInsights/Private/SSlateTraceFlags.cpp
--- a/Engine/Plugins/Slate/SlateInsights/Source/SlateInsights/Private/SSlateTraceFlags.cpp
+++ b/Engine/Plugins/Slate/SlateInsights/Source/SlateInsights/Private/SSlateTraceFlags.cpp
@@ -81,21 +81,20 @@ void SSlateTraceWidgetUpdateFlags::Construct(const FArguments& InArgs)
 		"T : Tick : The widget was updated/ticked.\n"
 		"T : Active Timer Update : The widget had an active timer.\n"
 		"P : Repaint : The widget was dirty and was repainted.\n"
-		"V : Volatile Paint : The widget was volatile and was repainted."
-		"P : Volatile Prepass : The widget was volatile and did a SlatePrepass."));
+		"V : Volatile Paint : The widget was volatile and was repainted.");
 }
 
 void SSlateTraceInvalidateWidgetReasonFlags::Construct(const FArguments& InArgs)
 {
 	Reason = InArgs._Reason;
 	SetToolTipText(LOCTEXT("InvalidateWidgetReasonFlagsTooltip", 
+		"C : Child Order : A child was added or removed (this implies layout).\n"
 		"L : Layout : The widget desired size changed.\n"
 		"P : Paint : The widget needs repainting but nothing affecting its size.\n"
 		"V : Volatile : The widget volatility changed.\n"
-		"C : Child Order : A child was added or removed (this implies layout).\n"
 		"R : Render Transform : The widget render transform changed.\n"
-		"V : Visibility : The widget visibility changed (this implies layout)."
-		"A : Attribute Registration : Attributes got bound or unbound."
+		"V : Visibility : The widget visibility changed (this implies layout).\n"
+		"A : Attribute Registration : Attributes got bound or unbound.\n"
 		"P : Prepass : Re-cache desired size of all of this widget's children recursively."));
 }
 
@@ -107,9 +106,8 @@ int32 SSlateTraceWidgetUpdateFlags::OnPaint(const FPaintArgs& Args, const FGeome
 		EWidgetUpdateFlags::NeedsActiveTimerUpdate,
 		EWidgetUpdateFlags::NeedsRepaint,
 		EWidgetUpdateFlags::NeedsVolatilePaint,
-		EWidgetUpdateFlags::NeedsVolatilePrepass,
 	};
-	const FString Text = TEXT("TTPVP");
+	const FString Text = TEXT("TTPV");
 
 	Private::Paint<EWidgetUpdateFlags>(AllUpdateFlags, Text, UpdateFlagsValue, AllottedGeometry, OutDrawElements, LayerId, InWidgetStyle);
 
@@ -121,16 +119,16 @@ int32 SSlateTraceInvalidateWidgetReasonFlags::OnPaint(const FPaintArgs& Args, co
 	FSlateWindowElementList& OutDrawElements, int32 LayerId, const FWidgetStyle& InWidgetStyle, bool bParentEnabled) const
 {
 	const EInvalidateWidgetReason AllReasons[] = {
+		EInvalidateWidgetReason::ChildOrder,
 		EInvalidateWidgetReason::Layout,
 		EInvalidateWidgetReason::Paint,
 		EInvalidateWidgetReason::Volatility,
-		EInvalidateWidgetReason::ChildOrder,
 		EInvalidateWidgetReason::RenderTransform,
 		EInvalidateWidgetReason::Visibility,
 		EInvalidateWidgetReason::AttributeRegistration,
 		EInvalidateWidgetReason::Prepass,
 	};
-	const FString Text = TEXT("LPVCRVAP");
+	const FString Text = TEXT("CLPVRVAP");
 
 	Private::Paint<EInvalidateWidgetReason>(AllReasons, Text, Reason, AllottedGeometry, OutDrawElements, LayerId, InWidgetStyle);
 
@@ -145,7 +143,7 @@ FVector2D SSlateTraceWidgetUpdateFlags::ComputeDesiredSize(float LayoutScaleMult
 
 FVector2D SSlateTraceInvalidateWidgetReasonFlags::ComputeDesiredSize(float LayoutScaleMultiplier) const
 {
-	return FVector2D(6 * Private::BoxSize.X, Private::BoxSize.Y);
+	return FVector2D(8 * Private::BoxSize.X, Private::BoxSize.Y);
 }
 
 } //namespace SlateInsights
diff --git a/Engine/Plugins/Slate/SlateInsights/Source/SlateInsights/Private/SlateAnalyzer.cpp b/Engine/Plugins/Slate/SlateInsights/Source/SlateInsights/Private/SlateAnalyzer.cpp
--- a/Engine/Plugins/Slate/SlateInsights/Source/SlateInsights/Private/SlateAnalyzer.cpp
+++ b/Engine/Plugins/Slate/SlateInsights/Source/SlateInsights/Private/SlateAnalyzer.cpp
@@ -75,7 +75,7 @@ bool FSlateAnalyzer::OnEvent(uint16 RouteId, EStyle Style, const FOnEventContext
 	{
 		const double Time = Context.EventTime.AsSeconds(EventData.GetValue<uint64>("Cycle"));
 
-		SlateProvider.AddWidgetUpdatedEvent(Time, Message::FWidgetUpdatedMessage(EventData));
+		SlateProvider.AddWidgetUpdatedEvent(Time, Message::FWidgetUpdatedMessage(EventData, Context.EventTime));
 		break;
 	}
 	case RouteId_WidgetInvalidated:
diff --git a/Engine/Plugins/Slate/SlateInsights/Source/SlateInsights/Private/SlateFrameGraphTrack.cpp b/Engine/Plugins/Slate/SlateInsights/Source/SlateInsights/Private/SlateFrameGraphTrack.cpp
--- a/Engine/Plugins/Slate/SlateInsights/Source/SlateInsights/Private/SlateFrameGraphTrack.cpp
+++ b/Engine/Plugins/Slate/SlateInsights/Source/SlateInsights/Private/SlateFrameGraphTrack.cpp
@@ -422,14 +422,14 @@ void FSlateFrameGraphTrack::AddAllSeries(const ITimingTrackUpdateContext& Contex
 			Series->SetColor_Detail(StyleSet.GetColor("SlateGraph.Color.RepaintCount"));
 			AllSeries.Add(Series);
 		}
-		{
-			TSharedRef<Private::FSlateFrameGraphSeries> Series = MakeShared<Private::FSlateFrameGraphSeries>(&Message::FApplicationTickedMessage::VolatilePaintCount);
-			Series->SetVisibility(true);
-			Series->SetName(LOCTEXT("VolatilePaintName", "Volatile Paint Count"));
-			Series->SetDescription(LOCTEXT("VolatilePaintDescription", "Volatile Paint Count"));
-			Series->SetColor_Detail(StyleSet.GetColor("SlateGraph.Color.VolatilePaintCount"));
-			AllSeries.Add(Series);
-		}
+		//{
+		//	TSharedRef<Private::FSlateFrameGraphSeries> Series = MakeShared<Private::FSlateFrameGraphSeries>(&Message::FApplicationTickedMessage::VolatilePaintCount);
+		//	Series->SetVisibility(true);
+		//	Series->SetName(LOCTEXT("VolatilePaintName", "Volatile Paint Count"));
+		//	Series->SetDescription(LOCTEXT("VolatilePaintDescription", "Volatile Paint Count"));
+		//	Series->SetColor_Detail(StyleSet.GetColor("SlateGraph.Color.VolatilePaintCount"));
+		//	AllSeries.Add(Series);
+		//}
 		{
 			TSharedRef<Private::FSlateFrameGraphSeries> Series = MakeShared<Private::FSlateFrameGraphSeries>(&Message::FApplicationTickedMessage::PaintCount);
 			Series->SetVisibility(true);
diff --git a/Engine/Plugins/Slate/SlateInsights/Source/SlateInsights/Private/SlateProvider.cpp b/Engine/Plugins/Slate/SlateInsights/Source/SlateInsights/Private/SlateProvider.cpp
--- a/Engine/Plugins/Slate/SlateInsights/Source/SlateInsights/Private/SlateProvider.cpp
+++ b/Engine/Plugins/Slate/SlateInsights/Source/SlateInsights/Private/SlateProvider.cpp
@@ -23,7 +23,7 @@ namespace Message
 		, TickCount(EventData.GetValue<uint32>("TickCount"))
 		, TimerCount(EventData.GetValue<uint32>("TimerCount"))
 		, RepaintCount(EventData.GetValue<uint32>("RepaintCount"))
-		, VolatilePaintCount(EventData.GetValue<uint32>("VolatilePaintCount"))
+		//, VolatilePaintCount(EventData.GetValue<uint32>("VolatilePaintCount"))
 		, PaintCount(EventData.GetValue<uint32>("PaintCount"))
 		, InvalidateCount(EventData.GetValue<uint32>("InvalidateCount"))
 		, RootInvalidatedCount(EventData.GetValue<uint32>("RootInvalidatedCount"))
@@ -48,8 +48,10 @@ namespace Message
 		EventData.GetString("DebugInfo", DebugInfo);
 	}
 	
-	FWidgetUpdatedMessage::FWidgetUpdatedMessage(const UE::Trace::IAnalyzer::FEventData& EventData)
+	FWidgetUpdatedMessage::FWidgetUpdatedMessage(const UE::Trace::IAnalyzer::FEventData& EventData, const UE::Trace::IAnalyzer::FEventTime& EventTime)
 		: WidgetId(EventData.GetValue<uint64>("WidgetId"))
+		, Duration(EventTime.AsSeconds(EventData.GetValue<uint64>("CycleEnd")) - EventTime.AsSeconds(EventData.GetValue<uint64>("Cycle")))
+		, AffectedCount(EventData.GetValue<int32>("AffectedCount"))
 		, UpdateFlags(static_cast<EWidgetUpdateFlags>(EventData.GetValue<uint8>("UpdateFlags")))
 	{
 		static_assert(sizeof(EWidgetUpdateFlags) == sizeof(uint8), "EWidgetUpdateFlags is not a uint8");
diff --git a/Engine/Plugins/Slate/SlateInsights/Source/SlateInsights/Private/SlateProvider.h b/Engine/Plugins/Slate/SlateInsights/Source/SlateInsights/Private/SlateProvider.h
--- a/Engine/Plugins/Slate/SlateInsights/Source/SlateInsights/Private/SlateProvider.h
+++ b/Engine/Plugins/Slate/SlateInsights/Source/SlateInsights/Private/SlateProvider.h
@@ -57,10 +57,14 @@ struct FWidgetInfo
 struct FWidgetUpdatedMessage
 {
 	FWidgetId WidgetId;
+	/** How long the update took. */
+	double Duration;
+	/** Number of widget that was affected by the updated widget. */
+	int32 AffectedCount;
 	/** Flag that was set by an invalidation or on the widget directly. */
 	EWidgetUpdateFlags UpdateFlags;
 
-	FWidgetUpdatedMessage(const UE::Trace::IAnalyzer::FEventData& EventData);
+	FWidgetUpdatedMessage(const UE::Trace::IAnalyzer::FEventData& EventData, const UE::Trace::IAnalyzer::FEventTime& EventTime);
 };
 
 struct FWidgetInvalidatedMessage
diff --git a/Engine/Source/Runtime/SlateCore/Private/FastUpdate/WidgetProxy.cpp b/Engine/Source/Runtime/SlateCore/Private/FastUpdate/WidgetProxy.cpp
--- a/Engine/Source/Runtime/SlateCore/Private/FastUpdate/WidgetProxy.cpp
+++ b/Engine/Source/Runtime/SlateCore/Private/FastUpdate/WidgetProxy.cpp
@@ -57,6 +57,7 @@ int32 FWidgetProxy::Update(const FPaintArgs& PaintArgs, FSlateWindowElementList&
 	}
 	else
 	{
+		UE_TRACE_SCOPED_SLATE_WIDGET_UPDATE(CurrentWidget.Get());
 		EWidgetUpdateFlags PreviousUpdateFlag = CurrentWidget->UpdateFlags;
 		if (CurrentWidget->HasAnyUpdateFlags(EWidgetUpdateFlags::NeedsActiveTimerUpdate))
 		{
@@ -77,7 +78,6 @@ int32 FWidgetProxy::Update(const FPaintArgs& PaintArgs, FSlateWindowElementList&
 #if WITH_SLATE_DEBUGGING
 		FSlateDebugging::BroadcastWidgetUpdated(CurrentWidget.Get(), PreviousUpdateFlag);
 #endif
-		UE_TRACE_SLATE_WIDGET_UPDATED(CurrentWidget.Get(), PreviousUpdateFlag);
 	}
 
 	return OutgoingLayerId;
diff --git a/Engine/Source/Runtime/SlateCore/Private/Trace/SlateTrace.cpp b/Engine/Source/Runtime/SlateCore/Private/Trace/SlateTrace.cpp
--- a/Engine/Source/Runtime/SlateCore/Private/Trace/SlateTrace.cpp
+++ b/Engine/Source/Runtime/SlateCore/Private/Trace/SlateTrace.cpp
@@ -36,7 +36,6 @@ UE_TRACE_EVENT_BEGIN(SlateTrace, ApplicationTickAndDrawWidgets)
 	UE_TRACE_EVENT_FIELD(uint32, TickCount)				// Amount of widget that needed to be tick.
 	UE_TRACE_EVENT_FIELD(uint32, TimerCount)			// Amount of widget that needed a timer update.
 	UE_TRACE_EVENT_FIELD(uint32, RepaintCount)			// Amount of widget that requested a paint.
-	UE_TRACE_EVENT_FIELD(uint32, VolatilePaintCount)	// Amount of widget that will always get painted.
 	UE_TRACE_EVENT_FIELD(uint32, PaintCount)			// Total amount of widget that got painted
 														//This can be higher than RepaintCount+VolatilePaintCount because some widget can get be painted as a side effect of another widget being painted.
 	UE_TRACE_EVENT_FIELD(uint32, InvalidateCount)		// Amount of widget that got invalidated.
@@ -63,7 +62,9 @@ UE_TRACE_EVENT_END()
 
 UE_TRACE_EVENT_BEGIN(SlateTrace, WidgetUpdated)
 	UE_TRACE_EVENT_FIELD(uint64, Cycle)
+	UE_TRACE_EVENT_FIELD(uint64, CycleEnd)
 	UE_TRACE_EVENT_FIELD(uint64, WidgetId)				// Updated widget unique ID.
+	UE_TRACE_EVENT_FIELD(uint32, AffectedCount)			// The number of widget that got affected by the object
 	UE_TRACE_EVENT_FIELD(uint8, UpdateFlags)			// The reason of the update. (EWidgetUpdateFlags)
 UE_TRACE_EVENT_END()
 
@@ -100,12 +101,13 @@ namespace SlateTraceDetail
 {
 	uint32 GWidgetCount = 0;
 	uint32 GScopedPaintCount = 0;
+	uint32 GScopedUpdateCount = 0;
 
 	uint32 GFramePaintCount = 0;
 	uint32 GFrameTickCount = 0;
 	uint32 GFrameTimerCount = 0;
 	uint32 GFrameRepaintCount = 0;
-	uint32 GFrameVolatileCount = 0;
+	//uint32 GFrameVolatileCount = 0;
 	uint32 GFrameInvalidateCount = 0;
 	uint32 GFrameRootInvalidateCount = 0;
 
@@ -130,18 +132,58 @@ namespace SlateTraceDetail
  //-----------------------------------------------------------------------------------//
 
 FSlateTrace::FScopedWidgetPaintTrace::FScopedWidgetPaintTrace(const SWidget* InWidget)
-	: StartCycle(FPlatformTime::Cycles64())
-	, Widget(InWidget)
-	, StartPaintCount(SlateTraceDetail::GScopedPaintCount)
+	: Widget(InWidget)
+	, StartPaintCount(SlateTraceDetail::GFramePaintCount)
 {
+	if (UE_TRACE_CHANNELEXPR_IS_ENABLED(SlateChannel))
+	{
+		StartCycle = FPlatformTime::Cycles64();
+	}
+
 	++SlateTraceDetail::GScopedPaintCount;
 	++SlateTraceDetail::GFramePaintCount;
 }
 
 FSlateTrace::FScopedWidgetPaintTrace::~FScopedWidgetPaintTrace()
 {
-	FSlateTrace::OutputWidgetPaint(Widget, StartCycle, FPlatformTime::Cycles64(), SlateTraceDetail::GScopedPaintCount - StartPaintCount);
 	--SlateTraceDetail::GScopedPaintCount;
+	if (UE_TRACE_CHANNELEXPR_IS_ENABLED(SlateChannel))
+	{
+		const double EndCycle = FPlatformTime::Cycles64();
+		const uint32 AffectedCount = SlateTraceDetail::GFramePaintCount - StartPaintCount;
+		//FSlateTrace::OutputWidgetPaint(Widget, StartCycle, EndCycle, AffectedCount);
+		if (SlateTraceDetail::GScopedPaintCount == 0)
+		{
+			const EWidgetUpdateFlags UpdateFlags = Widget->IsVolatile() ? EWidgetUpdateFlags::NeedsVolatilePaint : EWidgetUpdateFlags::NeedsRepaint;
+			FSlateTrace::OutputWidgetUpdate(Widget, StartCycle, EndCycle, UpdateFlags, AffectedCount);
+		}
+	}
+}
+
+//-----------------------------------------------------------------------------------//
+
+FSlateTrace::FScopedWidgetUpdateTrace::FScopedWidgetUpdateTrace(const SWidget* InWidget)
+{
+	if (UE_TRACE_CHANNELEXPR_IS_ENABLED(SlateChannel))
+	{
+		UpdateFlags = EWidgetUpdateFlags::None;
+		if (EnumHasAnyFlags(InWidget->UpdateFlags, EWidgetUpdateFlags::NeedsActiveTimerUpdate | EWidgetUpdateFlags::NeedsTick))
+		{
+			StartCycle = FPlatformTime::Cycles64();
+			Widget = InWidget;
+			UpdateFlags = InWidget->UpdateFlags & (EWidgetUpdateFlags::NeedsTick | EWidgetUpdateFlags::NeedsActiveTimerUpdate);
+		}
+	}
+	++SlateTraceDetail::GScopedUpdateCount;
+}
+
+FSlateTrace::FScopedWidgetUpdateTrace::~FScopedWidgetUpdateTrace()
+{
+	--SlateTraceDetail::GScopedUpdateCount;
+	if (UpdateFlags != EWidgetUpdateFlags::None)
+	{
+		FSlateTrace::OutputWidgetUpdate(Widget, StartCycle, FPlatformTime::Cycles64(), UpdateFlags, 1);
+	}
 }
 
 //-----------------------------------------------------------------------------------//
@@ -163,7 +205,7 @@ void FSlateTrace::ApplicationTickAndDrawWidgets(float DeltaTime)
 			<< ApplicationTickAndDrawWidgets.TickCount(SlateTraceDetail::GFrameTickCount)
 			<< ApplicationTickAndDrawWidgets.TimerCount(SlateTraceDetail::GFrameTimerCount)
 			<< ApplicationTickAndDrawWidgets.RepaintCount(SlateTraceDetail::GFrameRepaintCount)
-			<< ApplicationTickAndDrawWidgets.VolatilePaintCount(SlateTraceDetail::GFrameVolatileCount)
+			//<< ApplicationTickAndDrawWidgets.VolatilePaintCount(SlateTraceDetail::GFrameVolatileCount)
 			<< ApplicationTickAndDrawWidgets.PaintCount(SlateTraceDetail::GFramePaintCount)
 			<< ApplicationTickAndDrawWidgets.InvalidateCount(SlateTraceDetail::GFrameInvalidateCount)
 			<< ApplicationTickAndDrawWidgets.RootInvalidatedCount(SlateTraceDetail::GFrameRootInvalidateCount)
@@ -172,32 +214,39 @@ void FSlateTrace::ApplicationTickAndDrawWidgets(float DeltaTime)
 		SlateTraceDetail::GFrameTickCount = 0;
 		SlateTraceDetail::GFrameTimerCount = 0;
 		SlateTraceDetail::GFrameRepaintCount = 0;
-		SlateTraceDetail::GFrameVolatileCount = 0;
+		//SlateTraceDetail::GFrameVolatileCount = 0;
 		SlateTraceDetail::GFramePaintCount = 0;
 		SlateTraceDetail::GFrameInvalidateCount = 0;
 		SlateTraceDetail::GFrameRootInvalidateCount = 0;
 	}
 }
 
-void FSlateTrace::WidgetUpdated(const SWidget* Widget, EWidgetUpdateFlags UpdateFlags)
+void FSlateTrace::OutputWidgetUpdate(const SWidget* Widget, uint64 StartCycle, uint64 EndCycle, EWidgetUpdateFlags UpdateFlags, uint32 AffectedCount)
 {
-	if (UE_TRACE_CHANNELEXPR_IS_ENABLED(SlateChannel) && UpdateFlags != EWidgetUpdateFlags::None)
+	if (EnumHasAnyFlags(UpdateFlags, EWidgetUpdateFlags::NeedsActiveTimerUpdate))
 	{
-		if (EnumHasAnyFlags(UpdateFlags, EWidgetUpdateFlags::NeedsTick)) { ++SlateTraceDetail::GFrameTickCount; }
-		if (EnumHasAnyFlags(UpdateFlags, EWidgetUpdateFlags::NeedsActiveTimerUpdate)) { ++SlateTraceDetail::GFrameTimerCount; }
+		++SlateTraceDetail::GFrameTimerCount;
+	}
+	else if (EnumHasAnyFlags(UpdateFlags, EWidgetUpdateFlags::NeedsTick))
+	{
+		++SlateTraceDetail::GFrameTickCount;
+	}
 
-		if (EnumHasAnyFlags(UpdateFlags, EWidgetUpdateFlags::NeedsVolatilePaint)) { ++SlateTraceDetail::GFrameVolatileCount; }
-		else if (EnumHasAnyFlags(UpdateFlags, EWidgetUpdateFlags::NeedsRepaint)) { ++SlateTraceDetail::GFrameRepaintCount; }
+	if (EnumHasAnyFlags(UpdateFlags, EWidgetUpdateFlags::NeedsVolatilePaint | EWidgetUpdateFlags::NeedsRepaint))
+	{
+		++SlateTraceDetail::GFrameRepaintCount;
+	}
 
-		static_assert(sizeof(EWidgetUpdateFlags) == sizeof(uint8), "EWidgetUpdateFlags is not a uint8");
+	static_assert(sizeof(EWidgetUpdateFlags) == sizeof(uint8), "EWidgetUpdateFlags is not a uint8");
 
-		const uint64 WidgetId = SlateTraceDetail::GetWidgetId(Widget);
+	const uint64 WidgetId = SlateTraceDetail::GetWidgetId(Widget);
 
-		UE_TRACE_LOG(SlateTrace, WidgetUpdated, SlateChannel)
-			<< WidgetUpdated.Cycle(FPlatformTime::Cycles64())
-			<< WidgetUpdated.WidgetId(WidgetId)
-			<< WidgetUpdated.UpdateFlags(static_cast<uint8>(UpdateFlags));
-	}
+	UE_TRACE_LOG(SlateTrace, WidgetUpdated, SlateChannel)
+		<< WidgetUpdated.Cycle(StartCycle)
+		<< WidgetUpdated.CycleEnd(EndCycle)
+		<< WidgetUpdated.WidgetId(WidgetId)
+		<< WidgetUpdated.AffectedCount(AffectedCount)
+		<< WidgetUpdated.UpdateFlags(static_cast<uint8>(UpdateFlags));
 }
 
 void FSlateTrace::WidgetInvalidated(const SWidget* Widget, const SWidget* Investigator, EInvalidateWidgetReason Reason)
diff --git a/Engine/Source/Runtime/SlateCore/Private/Widgets/SWidget.cpp b/Engine/Source/Runtime/SlateCore/Private/Widgets/SWidget.cpp
--- a/Engine/Source/Runtime/SlateCore/Private/Widgets/SWidget.cpp
+++ b/Engine/Source/Runtime/SlateCore/Private/Widgets/SWidget.cpp
@@ -1408,29 +1408,32 @@ int32 SWidget::Paint(const FPaintArgs& Args, const FGeometry& AllottedGeometry,
 	FGeometry DesktopSpaceGeometry = AllottedGeometry;
 	DesktopSpaceGeometry.AppendTransform(FSlateLayoutTransform(Args.GetWindowToDesktopTransform()));
 
-	if (HasAnyUpdateFlags(EWidgetUpdateFlags::NeedsActiveTimerUpdate))
 	{
-		if (bHasPendingAttributesInvalidation)
+		UE_TRACE_SCOPED_SLATE_WIDGET_UPDATE(this);
+		if (HasAnyUpdateFlags(EWidgetUpdateFlags::NeedsActiveTimerUpdate))
 		{
-			FSlateAttributeMetaData::ApplyDelayedInvalidation(*MutableThis);
-		}
+			if (bHasPendingAttributesInvalidation)
+			{
+				FSlateAttributeMetaData::ApplyDelayedInvalidation(*MutableThis);
+			}
 
-		SCOPE_CYCLE_COUNTER(STAT_SlateExecuteActiveTimers);
-		MutableThis->ExecuteActiveTimers(Args.GetCurrentTime(), Args.GetDeltaTime());
-	}
+			SCOPE_CYCLE_COUNTER(STAT_SlateExecuteActiveTimers);
+			MutableThis->ExecuteActiveTimers(Args.GetCurrentTime(), Args.GetDeltaTime());
+		}
 
-	if (HasAnyUpdateFlags(EWidgetUpdateFlags::NeedsTick))
-	{
-		if (bHasPendingAttributesInvalidation)
+		if (HasAnyUpdateFlags(EWidgetUpdateFlags::NeedsTick))
 		{
-			FSlateAttributeMetaData::ApplyDelayedInvalidation(*MutableThis);
-		}
+			if (bHasPendingAttributesInvalidation)
+			{
+				FSlateAttributeMetaData::ApplyDelayedInvalidation(*MutableThis);
+			}
 
-		INC_DWORD_STAT(STAT_SlateNumTickedWidgets);
+			INC_DWORD_STAT(STAT_SlateNumTickedWidgets);
 
-		SCOPE_CYCLE_COUNTER(STAT_SlateTickWidgets);
-		SCOPE_CYCLE_SWIDGET(this);
-		MutableThis->Tick(DesktopSpaceGeometry, Args.GetCurrentTime(), Args.GetDeltaTime());
+			SCOPE_CYCLE_COUNTER(STAT_SlateTickWidgets);
+			SCOPE_CYCLE_SWIDGET(this);
+			MutableThis->Tick(DesktopSpaceGeometry, Args.GetCurrentTime(), Args.GetDeltaTime());
+		}
 	}
 
 	if (bHasPendingAttributesInvalidation)
@@ -1631,7 +1634,6 @@ int32 SWidget::Paint(const FPaintArgs& Args, const FGeometry& AllottedGeometry,
 #if WITH_SLATE_DEBUGGING
 	FSlateDebugging::BroadcastWidgetUpdatedByPaint(this, PreviousUpdateFlag);
 #endif
-	UE_TRACE_SLATE_WIDGET_UPDATED(this, PreviousUpdateFlag);
 
 	return NewLayerId;
 }
diff --git a/Engine/Source/Runtime/SlateCore/Public/Trace/SlateTrace.h b/Engine/Source/Runtime/SlateCore/Public/Trace/SlateTrace.h
--- a/Engine/Source/Runtime/SlateCore/Public/Trace/SlateTrace.h
+++ b/Engine/Source/Runtime/SlateCore/Public/Trace/SlateTrace.h
@@ -47,9 +47,17 @@ public:
 		const SWidget* Widget;
 		int32 StartPaintCount;
 	};
+	struct FScopedWidgetUpdateTrace
+	{
+		FScopedWidgetUpdateTrace(const SWidget* InWidget);
+		~FScopedWidgetUpdateTrace();
+
+		uint64 StartCycle;
+		const SWidget* Widget;
+		EWidgetUpdateFlags UpdateFlags;
+	};
 
 	SLATECORE_API static void ApplicationTickAndDrawWidgets(float DeltaTime);
-	SLATECORE_API static void WidgetUpdated(const SWidget* Widget, EWidgetUpdateFlags UpdateFlags);
 	SLATECORE_API static void WidgetInvalidated(const SWidget* Widget, const SWidget* Investigator, EInvalidateWidgetReason Reason);
 	SLATECORE_API static void RootInvalidated(const SWidget* Widget, const SWidget* Investigator);
 	SLATECORE_API static void RootChildOrderInvalidated(const SWidget* Widget, const SWidget* Investigator);
@@ -60,6 +68,7 @@ public:
 
 private:
 	static void OutputWidgetPaint(const SWidget* Widget, uint64 StartCycle, uint64 EndCycle, uint32 PaintCount);
+	static void OutputWidgetUpdate(const SWidget* Widget, uint64 StartCycle, uint64 EndCycle, EWidgetUpdateFlags UpdateFlags, uint32 AffectedCount);
 };
 
 #define UE_TRACE_SLATE_BOOKMARK(Format, ...) \
@@ -82,8 +91,8 @@ private:
 #define UE_TRACE_SCOPED_SLATE_WIDGET_PAINT(Widget) \
 	FSlateTrace::FScopedWidgetPaintTrace _ScopedSlateWidgetPaintTrace(Widget);
 	
-#define UE_TRACE_SLATE_WIDGET_UPDATED(Widget, UpdateFlag) \
-	FSlateTrace::WidgetUpdated(Widget, UpdateFlag);
+#define UE_TRACE_SCOPED_SLATE_WIDGET_UPDATE(Widget) \
+	FSlateTrace::FScopedWidgetUpdateTrace _ScopedSlateWidgetUpdateTrace(Widget);
 
 #define UE_TRACE_SLATE_WIDGET_INVALIDATED(Widget, Investigator, InvalidateWidgetReason) \
 	FSlateTrace::WidgetInvalidated(Widget, Investigator, InvalidateWidgetReason);
@@ -102,7 +111,7 @@ private:
 #define UE_TRACE_SLATE_WIDGET_DEBUG_INFO(Widget)
 #define UE_TRACE_SLATE_WIDGET_REMOVED(Widget)
 #define UE_TRACE_SCOPED_SLATE_WIDGET_PAINT(Widget)
-#define UE_TRACE_SLATE_WIDGET_UPDATED(Widget, UpdateFlag)
+#define UE_TRACE_SCOPED_SLATE_WIDGET_UPDATE(Widget)
 #define UE_TRACE_SLATE_WIDGET_INVALIDATED(Widget, Investigator, InvalidateWidgetReason)
 #define UE_TRACE_SLATE_ROOT_INVALIDATED(Widget, Investigator)
 #define UE_TRACE_SLATE_ROOT_CHILDORDER_INVALIDATED(Widget, Investigator)
diff --git a/Engine/Source/Runtime/SlateCore/Public/Widgets/SWidget.h b/Engine/Source/Runtime/SlateCore/Public/Widgets/SWidget.h
--- a/Engine/Source/Runtime/SlateCore/Public/Widgets/SWidget.h
+++ b/Engine/Source/Runtime/SlateCore/Public/Widgets/SWidget.h
@@ -160,6 +160,7 @@ class SLATECORE_API SWidget
 	friend class FSlateInvalidationWidgetList;
 	friend class FSlateWindowElementList;
 	friend class SWindow;
+	friend class FSlateTrace;
 	friend struct FSlateCachedElementList;
 	template<class WidgetType, typename RequiredArgsPayloadType>
 	friend struct TSlateDecl;
