SlateInsights: Fix compile unit.

diff --git a/Engine/Plugins/Slate/SlateInsights/Source/SlateInsights/Private/SSlateTraceFlags.cpp b/Engine/Plugins/Slate/SlateInsights/Source/SlateInsights/Private/SSlateTraceFlags.cpp
--- a/Engine/Plugins/Slate/SlateInsights/Source/SlateInsights/Private/SSlateTraceFlags.cpp
+++ b/Engine/Plugins/Slate/SlateInsights/Source/SlateInsights/Private/SSlateTraceFlags.cpp
@@ -81,7 +81,7 @@ void SSlateTraceWidgetUpdateFlags::Construct(const FArguments& InArgs)
 		"T : Tick : The widget was updated/ticked.\n"
 		"T : Active Timer Update : The widget had an active timer.\n"
 		"P : Repaint : The widget was dirty and was repainted.\n"
-		"V : Volatile Paint : The widget was volatile and was repainted.");
+		"V : Volatile Paint : The widget was volatile and was repainted."));
 }
 
 void SSlateTraceInvalidateWidgetReasonFlags::Construct(const FArguments& InArgs)
