Fix slate trace selection not updating under timing insights, requring clicking on asset loading insights to update.

diff --git a/Engine/Plugins/Slate/SlateInsights/Source/SlateInsights/Private/SlateTimingViewSession.cpp b/Engine/Plugins/Slate/SlateInsights/Source/SlateInsights/Private/SlateTimingViewSession.cpp
--- a/Engine/Plugins/Slate/SlateInsights/Source/SlateInsights/Private/SlateTimingViewSession.cpp
+++ b/Engine/Plugins/Slate/SlateInsights/Source/SlateInsights/Private/SlateTimingViewSession.cpp
@@ -3,6 +3,7 @@
 #include "SlateTimingViewSession.h"
 
 #include "Insights/ITimingViewSession.h"
+#include "Insights/IUnrealInsightsModule.h"
 #include "Insights/ViewModels/TimingEvent.h"
 #include "Modules/ModuleManager.h"
 #include "TraceServices/Model/AnalysisSession.h"
@@ -33,13 +34,21 @@ FSlateTimingViewSession::FSlateTimingViewSession()
 
 void FSlateTimingViewSession::OnBeginSession(Insights::ITimingViewSession& InTimingViewSession)
 {
-	TimingViewSession = &InTimingViewSession;
+	if (InTimingViewSession.GetName() == FInsightsManagerTabs::TimingProfilerTabId)
+	{
+		TimingViewSession = &InTimingViewSession;
+	}
 
 	SlateFrameGraphTrack.Reset();
 }
 
 void FSlateTimingViewSession::OnEndSession(Insights::ITimingViewSession& InTimingViewSession)
 {
+	if (&InTimingViewSession != TimingViewSession)
+	{
+		return;
+	}
+
 	const bool bInvoke = false;
 	if (TSharedPtr<SSlateFrameSchematicView> SchematicView = FSlateInsightsModule::Get().GetSlateFrameSchematicViewTab(bInvoke))
 	{
@@ -51,6 +60,11 @@ void FSlateTimingViewSession::OnEndSession(Insights::ITimingViewSession& InTimin
 
 void FSlateTimingViewSession::Tick(Insights::ITimingViewSession& InTimingViewSession, const TraceServices::IAnalysisSession& InAnalysisSession)
 {
+	if (&InTimingViewSession != TimingViewSession)
+	{
+		return;
+	}
+
 	AnalysisSession = &InAnalysisSession;
 
 	const FSlateProvider* SlateProvider = InAnalysisSession.ReadProvider<FSlateProvider>(FSlateProvider::ProviderName);
@@ -86,7 +100,6 @@ void FSlateTimingViewSession::Tick(Insights::ITimingViewSession& InTimingViewSes
 				SchematicView->SetSession(&InTimingViewSession, &InAnalysisSession);
 			}
 		}
-
 	}
 }
 
