Fix Slate Frame View not working when opened after insights timeline

diff --git a/Engine/Plugins/Slate/SlateInsights/Source/SlateInsights/Private/SSlateFrameSchematicView.cpp b/Engine/Plugins/Slate/SlateInsights/Source/SlateInsights/Private/SSlateFrameSchematicView.cpp
--- a/Engine/Plugins/Slate/SlateInsights/Source/SlateInsights/Private/SSlateFrameSchematicView.cpp
+++ b/Engine/Plugins/Slate/SlateInsights/Source/SlateInsights/Private/SSlateFrameSchematicView.cpp
@@ -371,6 +371,8 @@ void SSlateFrameSchematicView::Construct(const FArguments& InArgs)
 	WidgetUpdateSortColumn = FName();
 	bWidgetUpdateSortAscending = false;
 
+	TimingViewSession = nullptr;
+
 	ChildSlot
 	[
 		SNew(SVerticalBox)
@@ -566,6 +568,11 @@ void SSlateFrameSchematicView::SetSession(Insights::ITimingViewSession* InTiming
 	RefreshNodes();
 }
 
+bool SSlateFrameSchematicView::IsSessionSet() const
+{
+	return TimingViewSession != nullptr;
+}
+
 TSharedRef<ITableRow> SSlateFrameSchematicView::HandleUniqueInvalidatedMakeTreeRowWidget(TSharedPtr<Private::FWidgetUniqueInvalidatedInfo> InInfo, const TSharedRef<STableViewBase>& OwnerTable)
 {
 	return SNew(Private::FWidgetUniqueInvalidatedInfoRow,
diff --git a/Engine/Plugins/Slate/SlateInsights/Source/SlateInsights/Private/SSlateFrameSchematicView.h b/Engine/Plugins/Slate/SlateInsights/Source/SlateInsights/Private/SSlateFrameSchematicView.h
--- a/Engine/Plugins/Slate/SlateInsights/Source/SlateInsights/Private/SSlateFrameSchematicView.h
+++ b/Engine/Plugins/Slate/SlateInsights/Source/SlateInsights/Private/SSlateFrameSchematicView.h
@@ -41,6 +41,8 @@ public:
 
 	void SetSession(Insights::ITimingViewSession* InTimingViewSession, const TraceServices::IAnalysisSession* InAnalysisSession);
 
+	bool IsSessionSet() const;
+
 private:
 	TSharedRef<ITableRow> HandleUniqueInvalidatedMakeTreeRowWidget(TSharedPtr<Private::FWidgetUniqueInvalidatedInfo> InInfo, const TSharedRef<STableViewBase>& OwnerTable);
 	void HandleUniqueInvalidatedChildrenForInfo(TSharedPtr<Private::FWidgetUniqueInvalidatedInfo> InInfo, TArray<TSharedPtr<Private::FWidgetUniqueInvalidatedInfo>>& OutChildren);
diff --git a/Engine/Plugins/Slate/SlateInsights/Source/SlateInsights/Private/SlateTimingViewSession.cpp b/Engine/Plugins/Slate/SlateInsights/Source/SlateInsights/Private/SlateTimingViewSession.cpp
--- a/Engine/Plugins/Slate/SlateInsights/Source/SlateInsights/Private/SlateTimingViewSession.cpp
+++ b/Engine/Plugins/Slate/SlateInsights/Source/SlateInsights/Private/SlateTimingViewSession.cpp
@@ -95,6 +95,15 @@ void FSlateTimingViewSession::Tick(Insights::ITimingViewSession& InTimingViewSes
 				SchematicView->SetSession(&InTimingViewSession, &InAnalysisSession);
 			}
 		}
+
+		const bool bInvoke = false;
+		if (TSharedPtr<SSlateFrameSchematicView> SchematicView = FSlateInsightsModule::Get().GetSlateFrameSchematicViewTab(bInvoke))
+		{
+			if (!SchematicView->IsSessionSet())
+			{
+				SchematicView->SetSession(&InTimingViewSession, &InAnalysisSession);
+			}
+		}
 	}
 }
 
