Slate: Deprecate SListView::ItemHeight and STreeViewItemHeight. ItemHeight and ItemWidth are only used when the panel is is as a tile or both value are assigned. With TreeView and ListView, you can't assigned both. The function are protected, a subclass of those could but it's up to them to expose the ItemHeight and ItemWidth. This behavior changed was introduced with CL 7634094.

diff --git a/Engine/Plugins/Editor/GuidedTutorials/Source/IntroTutorials/Private/STutorialsBrowser.cpp b/Engine/Plugins/Editor/GuidedTutorials/Source/IntroTutorials/Private/STutorialsBrowser.cpp
--- a/Engine/Plugins/Editor/GuidedTutorials/Source/IntroTutorials/Private/STutorialsBrowser.cpp
+++ b/Engine/Plugins/Editor/GuidedTutorials/Source/IntroTutorials/Private/STutorialsBrowser.cpp
@@ -619,7 +619,6 @@ void STutorialsBrowser::Construct(const FArguments& InArgs)
 			.FillHeight(1.0f)
 			[
 				SAssignNew(TutorialList, SListView<TSharedPtr<ITutorialListEntry>>)
-				.ItemHeight(128.0f)
 				.ListItemsSource(&FilteredEntries)
 				.OnGenerateRow(this, &STutorialsBrowser::OnGenerateTutorialRow)
 				.SelectionMode(ESelectionMode::None)
diff --git a/Engine/Plugins/Experimental/AdvancedRenamer/Source/AdvancedRenamer/Private/Slate/SAdvancedRenamerPanel.cpp b/Engine/Plugins/Experimental/AdvancedRenamer/Source/AdvancedRenamer/Private/Slate/SAdvancedRenamerPanel.cpp
--- a/Engine/Plugins/Experimental/AdvancedRenamer/Source/AdvancedRenamer/Private/Slate/SAdvancedRenamerPanel.cpp
+++ b/Engine/Plugins/Experimental/AdvancedRenamer/Source/AdvancedRenamer/Private/Slate/SAdvancedRenamerPanel.cpp
@@ -49,7 +49,6 @@ namespace UE::AdvancedRenamer::Private
 	constexpr float RightBlockOffsetX = 314.f;
 	const FVector2D RightBlockSize = {281.f, 490.f};
 	const FVector2D ListViewSize = {277.f, 431.f};
-	constexpr float ListLineHeight = 15.f;
 	constexpr float ApplyButtonHeight = 25.f;
 	static FName OriginalNameColumnName = "OriginalName";
 	static FName NewNameColumnName = "NewName";
@@ -731,7 +730,6 @@ TSharedRef<SWidget> SAdvancedRenamerPanel::CreateRenamePreview()
 				.Content()
 				[
 					SAssignNew(RenamePreviewList, SListView<TSharedPtr<FAdvancedRenamerPreview>>)
-					.ItemHeight(ListLineHeight)
 					.ListItemsSource(&Renamer->GetPreviews())
 					.OnGenerateRow(this, &SAdvancedRenamerPanel::OnGenerateRowForList)
 					.HeaderRow(
diff --git a/Engine/Plugins/Experimental/Avalanche/Source/AvalancheMediaEditor/Private/Broadcast/OutputDevices/Slate/SAvaBroadcastOutputDevices.cpp b/Engine/Plugins/Experimental/Avalanche/Source/AvalancheMediaEditor/Private/Broadcast/OutputDevices/Slate/SAvaBroadcastOutputDevices.cpp
--- a/Engine/Plugins/Experimental/Avalanche/Source/AvalancheMediaEditor/Private/Broadcast/OutputDevices/Slate/SAvaBroadcastOutputDevices.cpp
+++ b/Engine/Plugins/Experimental/Avalanche/Source/AvalancheMediaEditor/Private/Broadcast/OutputDevices/Slate/SAvaBroadcastOutputDevices.cpp
@@ -23,7 +23,6 @@ void SAvaBroadcastOutputDevices::Construct(const FArguments& InArgs, const TShar
 	RootItem = MakeShared<FAvaBroadcastOutputRootItem>();
 	
 	SAssignNew(OutputTree, STreeView<FAvaOutputTreeItemPtr>)
-		.ItemHeight(20.0f)
 		.SelectionMode(ESelectionMode::Single)
 		.OnGenerateRow(this, &SAvaBroadcastOutputDevices::OnGenerateItemRow)
 		.OnGetChildren(this, &SAvaBroadcastOutputDevices::OnGetRowChildren)
diff --git a/Engine/Plugins/Experimental/Avalanche/Source/AvalancheTagEditor/Private/Customization/SAvaTagPicker.cpp b/Engine/Plugins/Experimental/Avalanche/Source/AvalancheTagEditor/Private/Customization/SAvaTagPicker.cpp
--- a/Engine/Plugins/Experimental/Avalanche/Source/AvalancheTagEditor/Private/Customization/SAvaTagPicker.cpp
+++ b/Engine/Plugins/Experimental/Avalanche/Source/AvalancheTagEditor/Private/Customization/SAvaTagPicker.cpp
@@ -39,7 +39,6 @@ void SAvaTagPicker::Construct(const FArguments& InArgs, const TSharedRef<IProper
 		.OnGenerateRow(this, &SAvaTagPicker::CreateTagTableRow)
 		.SelectionMode(ESelectionMode::Single)
 		.ListItemsSource(&TagOptions)
-		.ItemHeight(24.0f)
 		.IsFocusable(false)
 		.HandleGamepadEvents(false)
 		.HandleSpacebarSelection(false)
diff --git a/Engine/Plugins/Experimental/HairCardGenerator/Source/HairCardGeneratorEditor/Private/HairCardGenOptionsWindow.cpp b/Engine/Plugins/Experimental/HairCardGenerator/Source/HairCardGeneratorEditor/Private/HairCardGenOptionsWindow.cpp
--- a/Engine/Plugins/Experimental/HairCardGenerator/Source/HairCardGeneratorEditor/Private/HairCardGenOptionsWindow.cpp
+++ b/Engine/Plugins/Experimental/HairCardGenerator/Source/HairCardGeneratorEditor/Private/HairCardGenOptionsWindow.cpp
@@ -528,7 +528,6 @@ void SHairCardGenOptionsWindow::Construct(const FArguments& InArgs)
 				.Padding(2)
 			[
 				SAssignNew(GroupSettingsListView, SListView<TObjectPtr<UHairCardGeneratorGroupSettings>>)
-					.ItemHeight(24)
 					.ListItemsSource(&SettingsObject->GetFilterGroupSettings())
 					.OnGenerateRow(this, &SHairCardGenOptionsWindow::OnGenerateOptionListRow)
 					.SelectionMode(this, &SHairCardGenOptionsWindow::GetListSelectionMode)
diff --git a/Engine/Plugins/Experimental/PropertyAnimatorCore/Source/PropertyAnimatorCoreEditor/Private/Widgets/Views/SPropertyAnimatorCoreEditorControllersView.cpp b/Engine/Plugins/Experimental/PropertyAnimatorCore/Source/PropertyAnimatorCoreEditor/Private/Widgets/Views/SPropertyAnimatorCoreEditorControllersView.cpp
--- a/Engine/Plugins/Experimental/PropertyAnimatorCore/Source/PropertyAnimatorCoreEditor/Private/Widgets/Views/SPropertyAnimatorCoreEditorControllersView.cpp
+++ b/Engine/Plugins/Experimental/PropertyAnimatorCore/Source/PropertyAnimatorCoreEditor/Private/Widgets/Views/SPropertyAnimatorCoreEditorControllersView.cpp
@@ -25,7 +25,6 @@ void SPropertyAnimatorCoreEditorControllersView::Construct(const FArguments& InA
 			SAssignNew(ControllersTree, STreeView<FControllersViewItemPtr>)
 			.TreeItemsSource(&ControllersTreeSource)
 			.SelectionMode(ESelectionMode::Multi)
-			.ItemHeight(30.f)
 			.HeaderRow(
 				SNew(SHeaderRow)
 				+ SHeaderRow::Column(SPropertyAnimatorCoreEditorEditPanel::HeaderAnimatorColumnName)
diff --git a/Engine/Plugins/Experimental/RenderGrid/Source/RenderGridEditor/Private/UI/SRenderGridJobList.cpp b/Engine/Plugins/Experimental/RenderGrid/Source/RenderGridEditor/Private/UI/SRenderGridJobList.cpp
--- a/Engine/Plugins/Experimental/RenderGrid/Source/RenderGridEditor/Private/UI/SRenderGridJobList.cpp
+++ b/Engine/Plugins/Experimental/RenderGrid/Source/RenderGridEditor/Private/UI/SRenderGridJobList.cpp
@@ -370,7 +370,6 @@ void UE::RenderGrid::Private::SRenderGridJobListTable::Construct(const FArgument
 		SListView::FArguments()
 		.ListItemsSource(InArgs._ListItemsSource)
 		.HeaderRow(InArgs._HeaderRow)
-		.ItemHeight(20.0f)
 		.OnGenerateRow(this, &SRenderGridJobListTable::HandleJobListGenerateRow)
 		.OnSelectionChanged(this, &SRenderGridJobListTable::HandleJobListSelectionChanged)
 		.SelectionMode(ESelectionMode::Multi)
diff --git a/Engine/Plugins/Experimental/StormSync/Source/StormSyncEditor/Private/Slate/SStormSyncReportDialog.cpp b/Engine/Plugins/Experimental/StormSync/Source/StormSyncEditor/Private/Slate/SStormSyncReportDialog.cpp
--- a/Engine/Plugins/Experimental/StormSync/Source/StormSyncEditor/Private/Slate/SStormSyncReportDialog.cpp
+++ b/Engine/Plugins/Experimental/StormSync/Source/StormSyncEditor/Private/Slate/SStormSyncReportDialog.cpp
@@ -127,7 +127,6 @@ void SStormSyncReportDialog::Construct(const FArguments& InArgs, const FText& In
 				[
 					SAssignNew(ReportTreeView, SStormSyncReportTreeView)
 					.TreeItemsSource(&PackageReportRootNode.Children)
-					.ItemHeight(18)
 					.SelectionMode(ESelectionMode::Single)
 					.OnGenerateRow(this, &SStormSyncReportDialog::GenerateTreeRow)
 					.OnGetChildren_Static(&SStormSyncReportDialog::GetChildrenForTree)
diff --git a/Engine/Plugins/Experimental/UserToolBoxCore/Source/UserToolBoxCore/Private/Editor/SCommandPickMenu.cpp b/Engine/Plugins/Experimental/UserToolBoxCore/Source/UserToolBoxCore/Private/Editor/SCommandPickMenu.cpp
--- a/Engine/Plugins/Experimental/UserToolBoxCore/Source/UserToolBoxCore/Private/Editor/SCommandPickMenu.cpp
+++ b/Engine/Plugins/Experimental/UserToolBoxCore/Source/UserToolBoxCore/Private/Editor/SCommandPickMenu.cpp
@@ -32,7 +32,6 @@ void SCommandPickMenuWidget::Construct(const FArguments& InArgs)
 	
 	SAssignNew(TreeViewPtr,STreeView<TSharedPtr<FCommandTreeNode>>)
 		.AllowOverscroll(EAllowOverscroll::Yes)
-		.ItemHeight(24)
 		.TreeItemsSource(&CommandTree)
 		.OnSelectionChanged_Lambda([this](TSharedPtr<FCommandTreeNode> Node, ESelectInfo::Type)
 		{
diff --git a/Engine/Plugins/Experimental/UserToolBoxCore/Source/UserToolBoxCore/Private/Editor/UTBTabEditor.cpp b/Engine/Plugins/Experimental/UserToolBoxCore/Source/UserToolBoxCore/Private/Editor/UTBTabEditor.cpp
--- a/Engine/Plugins/Experimental/UserToolBoxCore/Source/UserToolBoxCore/Private/Editor/UTBTabEditor.cpp
+++ b/Engine/Plugins/Experimental/UserToolBoxCore/Source/UserToolBoxCore/Private/Editor/UTBTabEditor.cpp
@@ -824,7 +824,6 @@ void FUTBTabEditor::CreateOrRebuildUTBTabSectionTabWidget()
 	.AutoHeight()
 	[		// Groups List
 		SAssignNew( SectionListView,SListView<TSharedPtr<FString>>)
-			.ItemHeight(24.f)
 			.OnGenerateRow_Lambda(
 
 			[this] (const TSharedPtr<FString> SectionName, const TSharedRef<STableViewBase>& OwnerTable)
diff --git a/Engine/Plugins/Experimental/UserToolBoxCore/Source/UserToolBoxCore/Private/UTBFactory.cpp b/Engine/Plugins/Experimental/UserToolBoxCore/Source/UserToolBoxCore/Private/UTBFactory.cpp
--- a/Engine/Plugins/Experimental/UserToolBoxCore/Source/UserToolBoxCore/Private/UTBFactory.cpp
+++ b/Engine/Plugins/Experimental/UserToolBoxCore/Source/UserToolBoxCore/Private/UTBFactory.cpp
@@ -176,7 +176,6 @@ bool UUTBCommandFactory::ConfigureProperties()
 		+SScrollBox::Slot()
 		[
 			SNew(STreeView<TSharedPtr<TreeNode>>)
-			.ItemHeight(24)
 			.TreeItemsSource(&Tree)
 			.OnGenerateRow(STreeView<TSharedPtr<TreeNode>>::FOnGenerateRow::CreateLambda([](TSharedPtr<TreeNode> Node,const TSharedRef<STableViewBase>& OwnerTable)
 			{
diff --git a/Engine/Plugins/Runtime/nDisplay/Source/DisplayClusterColorGrading/Private/SDisplayClusterColorGradingObjectList.cpp b/Engine/Plugins/Runtime/nDisplay/Source/DisplayClusterColorGrading/Private/SDisplayClusterColorGradingObjectList.cpp
--- a/Engine/Plugins/Runtime/nDisplay/Source/DisplayClusterColorGrading/Private/SDisplayClusterColorGradingObjectList.cpp
+++ b/Engine/Plugins/Runtime/nDisplay/Source/DisplayClusterColorGrading/Private/SDisplayClusterColorGradingObjectList.cpp
@@ -188,7 +188,6 @@ void SDisplayClusterColorGradingObjectList::Construct(const FArguments& InArgs)
 		[
 			SAssignNew(ListView, SListView<FDisplayClusterColorGradingListItemRef>)
 			.ListItemsSource(InArgs._ColorGradingItemsSource)
-			.ItemHeight(28)
 			.SelectionMode(ESelectionMode::Single)
 			.OnGenerateRow(this, &SDisplayClusterColorGradingObjectList::GenerateListItemRow)
 			.OnSelectionChanged(this, &SDisplayClusterColorGradingObjectList::OnSelectionChanged)
diff --git a/Engine/Plugins/Slate/SlateInsights/Source/SlateInsights/Private/SSlateFrameSchematicView.cpp b/Engine/Plugins/Slate/SlateInsights/Source/SlateInsights/Private/SSlateFrameSchematicView.cpp
--- a/Engine/Plugins/Slate/SlateInsights/Source/SlateInsights/Private/SSlateFrameSchematicView.cpp
+++ b/Engine/Plugins/Slate/SlateInsights/Source/SlateInsights/Private/SSlateFrameSchematicView.cpp
@@ -422,7 +422,6 @@ void SSlateFrameSchematicView::Construct(const FArguments& InArgs)
 				.FillHeight(1.f)
 				[
 					SAssignNew(WidgetInvalidateInfoListView, STreeView<TSharedPtr<Private::FWidgetUniqueInvalidatedInfo>>)
-					.ItemHeight(24.0f)
 					.TreeItemsSource(&WidgetInvalidationInfos)
 					.OnGenerateRow(this, &SSlateFrameSchematicView::HandleUniqueInvalidatedMakeTreeRowWidget)
 					.OnGetChildren(this, &SSlateFrameSchematicView::HandleUniqueInvalidatedChildrenForInfo)
@@ -481,7 +480,6 @@ void SSlateFrameSchematicView::Construct(const FArguments& InArgs)
 				[
 					SAssignNew(WidgetUpdateInfoListView, SListView<TSharedPtr<Private::FWidgetUpdateInfo>>)
 					.ScrollbarVisibility(EVisibility::Visible)
-					.ItemHeight(24.0f)
 					.ListItemsSource(&WidgetUpdateInfos)
 					.SelectionMode(ESelectionMode::SingleToggle)
 					.OnGenerateRow(this, &SSlateFrameSchematicView::HandleWidgetUpdateInfoGenerateWidget)
diff --git a/Engine/Plugins/Tests/AutomationDriverTests/Source/AutomationDriverTests/Private/SAutomationDriverSpecSuite.cpp b/Engine/Plugins/Tests/AutomationDriverTests/Source/AutomationDriverTests/Private/SAutomationDriverSpecSuite.cpp
--- a/Engine/Plugins/Tests/AutomationDriverTests/Source/AutomationDriverTests/Private/SAutomationDriverSpecSuite.cpp
+++ b/Engine/Plugins/Tests/AutomationDriverTests/Source/AutomationDriverTests/Private/SAutomationDriverSpecSuite.cpp
@@ -142,7 +142,6 @@ public:
 			.Padding(5, 5)
 			[
 				SAssignNew(DocumentList, SListView<TSharedRef<FDocumentInfo>>)
-				.ItemHeight(24)
 				.SelectionMode(ESelectionMode::None)
 				.ListItemsSource(&InViewModel->GetDocuments())
 				.OnGenerateRow(this, &SAutomationDriverSpecSuiteImpl::GenerateListRow)
diff --git a/Engine/Plugins/VirtualProduction/DMX/DMXEngine/Source/DMXEditor/Private/Widgets/FixturePatch/SDMXReadOnlyFixturePatchList.cpp b/Engine/Plugins/VirtualProduction/DMX/DMXEngine/Source/DMXEditor/Private/Widgets/FixturePatch/SDMXReadOnlyFixturePatchList.cpp
--- a/Engine/Plugins/VirtualProduction/DMX/DMXEngine/Source/DMXEditor/Private/Widgets/FixturePatch/SDMXReadOnlyFixturePatchList.cpp
+++ b/Engine/Plugins/VirtualProduction/DMX/DMXEngine/Source/DMXEditor/Private/Widgets/FixturePatch/SDMXReadOnlyFixturePatchList.cpp
@@ -118,7 +118,6 @@ void SDMXReadOnlyFixturePatchList::Construct(const FArguments& InArgs)
 			[
 				SAssignNew(ListView, SListView<TSharedPtr<FDMXReadOnlyFixturePatchListItem>>)
 				.HeaderRow(GenerateHeaderRow())
-				.ItemHeight(60.0f)
 				.ListItemsSource(&ListItems)
 				.OnGenerateRow(this, &SDMXReadOnlyFixturePatchList::OnGenerateRow)
 				.OnContextMenuOpening(InArgs._OnContextMenuOpening)
diff --git a/Engine/Plugins/VirtualProduction/DMX/DMXEngine/Source/DMXEditor/Private/Widgets/FixtureType/SDMXFixtureTypeFunctionsEditor.cpp b/Engine/Plugins/VirtualProduction/DMX/DMXEngine/Source/DMXEditor/Private/Widgets/FixtureType/SDMXFixtureTypeFunctionsEditor.cpp
--- a/Engine/Plugins/VirtualProduction/DMX/DMXEngine/Source/DMXEditor/Private/Widgets/FixtureType/SDMXFixtureTypeFunctionsEditor.cpp
+++ b/Engine/Plugins/VirtualProduction/DMX/DMXEngine/Source/DMXEditor/Private/Widgets/FixtureType/SDMXFixtureTypeFunctionsEditor.cpp
@@ -123,7 +123,6 @@ void SDMXFixtureTypeFunctionsEditor::RebuildList()
 				(
 					SAssignNew(ListView, SListView<TSharedPtr<FDMXFixtureTypeFunctionsEditorItemBase>>)
 					.HeaderRow(GenerateHeaderRow())
-					.ItemHeight(40.0f)
 					.ListItemsSource(&ListSource)
 					.OnGenerateRow(this, &SDMXFixtureTypeFunctionsEditor::OnGenerateRow)
 					.OnSelectionChanged(this, &SDMXFixtureTypeFunctionsEditor::OnListSelectionChanged)
diff --git a/Engine/Plugins/VirtualProduction/DMX/DMXEngine/Source/DMXEditor/Private/Widgets/FixtureType/SDMXFixtureTypeModesEditor.cpp b/Engine/Plugins/VirtualProduction/DMX/DMXEngine/Source/DMXEditor/Private/Widgets/FixtureType/SDMXFixtureTypeModesEditor.cpp
--- a/Engine/Plugins/VirtualProduction/DMX/DMXEngine/Source/DMXEditor/Private/Widgets/FixtureType/SDMXFixtureTypeModesEditor.cpp
+++ b/Engine/Plugins/VirtualProduction/DMX/DMXEngine/Source/DMXEditor/Private/Widgets/FixtureType/SDMXFixtureTypeModesEditor.cpp
@@ -121,7 +121,6 @@ void SDMXFixtureTypeModesEditor::RebuildList()
 			ListContentBorder->SetContent
 			(
 				SAssignNew(ListView, SListView<TSharedPtr<FDMXFixtureTypeModesEditorModeItem>>)
-				.ItemHeight(40.0f)
 				.ListItemsSource(&ListSource)
 				.OnGenerateRow(this, &SDMXFixtureTypeModesEditor::OnGenerateModeRow)
 				.OnSelectionChanged(this, &SDMXFixtureTypeModesEditor::OnListSelectionChanged)
diff --git a/Engine/Plugins/VirtualProduction/DMX/DMXEngine/Source/DMXEditor/Private/Widgets/Monitors/SDMXActivityInUniverse.cpp b/Engine/Plugins/VirtualProduction/DMX/DMXEngine/Source/DMXEditor/Private/Widgets/Monitors/SDMXActivityInUniverse.cpp
--- a/Engine/Plugins/VirtualProduction/DMX/DMXEngine/Source/DMXEditor/Private/Widgets/Monitors/SDMXActivityInUniverse.cpp
+++ b/Engine/Plugins/VirtualProduction/DMX/DMXEngine/Source/DMXEditor/Private/Widgets/Monitors/SDMXActivityInUniverse.cpp
@@ -100,8 +100,7 @@ void SDMXActivityInUniverse::Construct(const FArguments& InArgs)
 					.OnGenerateRow(this, &SDMXActivityInUniverse::GenerateChannelRow)
 					.ListItemsSource(&ChannelListSource)
 					.Orientation(EOrientation::Orient_Horizontal)
-					.ScrollbarVisibility(EVisibility::Collapsed)					
-					.ItemHeight(40.0f)
+					.ScrollbarVisibility(EVisibility::Collapsed)
 				]
 			]
 			+ SVerticalBox::Slot()
diff --git a/Engine/Plugins/VirtualProduction/DMX/DMXEngine/Source/DMXEditor/Private/Widgets/Monitors/SDMXActivityMonitor.cpp b/Engine/Plugins/VirtualProduction/DMX/DMXEngine/Source/DMXEditor/Private/Widgets/Monitors/SDMXActivityMonitor.cpp
--- a/Engine/Plugins/VirtualProduction/DMX/DMXEngine/Source/DMXEditor/Private/Widgets/Monitors/SDMXActivityMonitor.cpp
+++ b/Engine/Plugins/VirtualProduction/DMX/DMXEngine/Source/DMXEditor/Private/Widgets/Monitors/SDMXActivityMonitor.cpp
@@ -205,7 +205,6 @@ void SDMXActivityMonitor::Construct(const FArguments& InArgs)
 						+ SScrollBox::Slot()				
 						[
 							SAssignNew(UniverseList, SListView<TSharedPtr<SDMXActivityInUniverse>>)
-							.ItemHeight(20.0f)
 							.ListItemsSource(&UniverseListSource)
 							.Visibility(EVisibility::Visible)
 							.OnGenerateRow(this, &SDMXActivityMonitor::OnGenerateUniverseRow)
diff --git a/Engine/Plugins/VirtualProduction/DMX/DMXEngine/Source/DMXEditor/Private/Widgets/SDMXEntityTreeViewBase.cpp b/Engine/Plugins/VirtualProduction/DMX/DMXEngine/Source/DMXEditor/Private/Widgets/SDMXEntityTreeViewBase.cpp
--- a/Engine/Plugins/VirtualProduction/DMX/DMXEngine/Source/DMXEditor/Private/Widgets/SDMXEntityTreeViewBase.cpp
+++ b/Engine/Plugins/VirtualProduction/DMX/DMXEngine/Source/DMXEditor/Private/Widgets/SDMXEntityTreeViewBase.cpp
@@ -84,7 +84,6 @@ void SDMXEntityTreeViewBase::Construct(const FArguments& InArgs)
 	// and also controls selection and drag/drop
 	RootNode = MakeShared<FDMXEntityTreeRootNode>();
 	EntitiesTreeWidget = SNew(STreeView<TSharedPtr<FDMXEntityTreeNodeBase>>)
-		.ItemHeight(24)
 		.TreeItemsSource(&RootNode->GetChildren())
 		.SelectionMode(ESelectionMode::Multi)
 		.OnGenerateRow(this, &SDMXEntityTreeViewBase::OnGenerateRow)
diff --git a/Engine/Plugins/VirtualProduction/TimedDataMonitor/Source/TimedDataMonitorEditor/Private/STimedDataMonitorBufferVisualizer.cpp b/Engine/Plugins/VirtualProduction/TimedDataMonitor/Source/TimedDataMonitorEditor/Private/STimedDataMonitorBufferVisualizer.cpp
--- a/Engine/Plugins/VirtualProduction/TimedDataMonitor/Source/TimedDataMonitorEditor/Private/STimedDataMonitorBufferVisualizer.cpp
+++ b/Engine/Plugins/VirtualProduction/TimedDataMonitor/Source/TimedDataMonitorEditor/Private/STimedDataMonitorBufferVisualizer.cpp
@@ -163,7 +163,6 @@ void STimedDataMonitorBufferVisualizer::Construct(const FArguments& InArgs)
 		.ListItemsSource(&ListItemsSource)
 		.OnGenerateRow(this, &STimedDataMonitorBufferVisualizer::MakeListViewWidget)
 		.OnRowReleased(this, &STimedDataMonitorBufferVisualizer::ReleaseListViewWidget)
-		.ItemHeight(60)
 		.HeaderRow(
 			SNew(SHeaderRow)
 			+ SHeaderRow::Column(TimedDataBufferVisualizer::HeaderIdName_DisplayName)
diff --git a/Engine/Source/Developer/Profiler/Private/Widgets/SEventGraph.cpp b/Engine/Source/Developer/Profiler/Private/Widgets/SEventGraph.cpp
--- a/Engine/Source/Developer/Profiler/Private/Widgets/SEventGraph.cpp
+++ b/Engine/Source/Developer/Profiler/Private/Widgets/SEventGraph.cpp
@@ -753,7 +753,6 @@ void SEventGraph::Construct( const FArguments& InArgs )
 											.OnGenerateRow(this, &SEventGraph::EventGraph_OnGenerateRow)
 											.OnSelectionChanged(this, &SEventGraph::EventGraph_OnSelectionChanged)
 											.OnContextMenuOpening(FOnContextMenuOpening::CreateSP(this, &SEventGraph::EventGraph_GetMenuContent))
-											.ItemHeight(12.0f)
 											.HeaderRow
 											(
 												SAssignNew(TreeViewHeaderRow,SHeaderRow)
diff --git a/Engine/Source/Developer/Profiler/Private/Widgets/SFiltersAndPresets.cpp b/Engine/Source/Developer/Profiler/Private/Widgets/SFiltersAndPresets.cpp
--- a/Engine/Source/Developer/Profiler/Private/Widgets/SFiltersAndPresets.cpp
+++ b/Engine/Source/Developer/Profiler/Private/Widgets/SFiltersAndPresets.cpp
@@ -751,7 +751,6 @@ void SFiltersAndPresets::Construct( const FArguments& InArgs )
 				.OnGenerateRow( this, &SFiltersAndPresets::GroupAndStatTree_OnGenerateRow )
 				.OnMouseButtonDoubleClick( this, &SFiltersAndPresets::GroupAndStatTree_OnMouseButtonDoubleClick )
 				//.OnSelectionChanged( this, &SFiltersAndPresets::GroupAndStatTree_OnSelectionChanged )
-				.ItemHeight( 12 )
 			]
 		]
 	];
diff --git a/Engine/Source/Developer/Profiler/Private/Widgets/SMultiDumpBrowser.cpp b/Engine/Source/Developer/Profiler/Private/Widgets/SMultiDumpBrowser.cpp
--- a/Engine/Source/Developer/Profiler/Private/Widgets/SMultiDumpBrowser.cpp
+++ b/Engine/Source/Developer/Profiler/Private/Widgets/SMultiDumpBrowser.cpp
@@ -40,7 +40,6 @@ void SMultiDumpBrowser::Construct(const FArguments& InArgs)
 				[
 					SAssignNew(FileList, SListView<TSharedPtr<FFileDescriptor>>)
 					.ListItemsSource(&StatsFiles)
-					.ItemHeight(16)
 					.OnGenerateRow(this, &SMultiDumpBrowser::GenerateFileRow)
 					.OnSelectionChanged(this, &SMultiDumpBrowser::SelectionChanged)
 				]
